Copying //UE4/Dev-Framework to Dev-Main (//UE4/Dev-Main)

diff --git a/Engine/Plugins/Developer/GameplayDebugger/Source/GameplayDebuggerModule/Private/GameplayDebuggerReplicator.cpp b/Engine/Plugins/Developer/GameplayDebugger/Source/GameplayDebuggerModule/Private/GameplayDebuggerReplicator.cpp
new file mode 100644
--- /dev/null
+++ b/Engine/Plugins/Developer/GameplayDebugger/Source/GameplayDebuggerModule/Private/GameplayDebuggerReplicator.cpp
@@ -0,0 +1,1297 @@
+// Copyright 1998-2015 Epic Games, Inc. All Rights Reserved.
+
+#include "GameplayDebuggerPluginPrivatePCH.h"
+#include "GameplayDebuggerModuleSettings.h"
+#include "DebugRenderSceneProxy.h"
+#include "ModuleManager.h"
+#include "Engine/Selection.h"
+#include "Engine/GameInstance.h"
+#include "Engine/DebugCameraController.h"
+#include "Debug/DebugDrawService.h"
+#include "GameFramework/HUD.h"
+#include "GameplayDebuggerReplicator.h"
+#include "BehaviorTreeDelegates.h"
+#if WITH_EDITOR
+#include "Editor/EditorEngine.h"
+#include "Editor/BehaviorTreeEditor/Public/BehaviorTreeEditorModule.h"
+#include "LevelEditorViewport.h"
+#endif // WITH_EDITOR
+#include "UnrealNetwork.h"
+
+namespace GameplayDebuggerHelpers
+{
+	const float ActivationKeyTimePch = 0.25;
+}
+
+AGameplayDebuggerReplicator::AGameplayDebuggerReplicator(const FObjectInitializer& ObjectInitializer)
+	: Super(ObjectInitializer)
+	, PlayerControllersUpdateDelay(0)
+	, ActivationKeyTime(0)
+{
+#if ENABLED_GAMEPLAY_DEBUGGER
+	// Structure to hold one-time initialization
+	struct FConstructorStatics
+	{
+		ConstructorHelpers::FObjectFinderOptional<UTexture2D> RedIcon;
+		ConstructorHelpers::FObjectFinderOptional<UTexture2D> GreenIcon;
+
+		// both icons are needed to debug AI and we hold static initialization for cooker
+		FConstructorStatics()
+			: RedIcon(TEXT("/Engine/EngineResources/AICON-Red.AICON-Red"))
+			, GreenIcon(TEXT("/Engine/EngineResources/AICON-Green.AICON-Green"))
+		{
+		}
+	};
+	static FConstructorStatics ConstructorStatics;
+
+	DefaultTexture_Red = ConstructorStatics.RedIcon.Get();
+	DefaultTexture_Green = ConstructorStatics.GreenIcon.Get();
+
+	PrimaryActorTick.bCanEverTick = true;
+	PrimaryActorTick.bAllowTickOnDedicatedServer = true;
+	PrimaryActorTick.bTickEvenWhenPaused = true;
+	PrimaryActorTick.bStartWithTickEnabled = false;
+
+	USceneComponent* SceneComponent = CreateDefaultSubobject<USceneComponent>(TEXT("SceneComponent"));
+	RootComponent = SceneComponent;
+
+	CreateDefaultSubobject<UGameplayDebuggerDrawComponent>(TEXT("GameplayDebuggerRenderingComponent"));
+
+#if WITH_EDITOR
+	SetIsTemporarilyHiddenInEditor(true);
+#endif
+
+#if WITH_EDITORONLY_DATA
+	bHiddenEdLevel = true;
+	bHiddenEdLayer = true;
+	bHiddenEd = true;
+	bEditable = false;
+#endif
+
+	MaxMenuWidth = 0;
+#endif //ENABLED_GAMEPLAY_DEBUGGER
+}
+
+void AGameplayDebuggerReplicator::GetLifetimeReplicatedProps(TArray< FLifetimeProperty > & OutLifetimeProps) const
+{
+	Super::GetLifetimeReplicatedProps(OutLifetimeProps);
+
+#if ENABLED_GAMEPLAY_DEBUGGER
+	DOREPLIFETIME_CONDITION(AGameplayDebuggerReplicator, LocalPlayerOwner, COND_OwnerOnly);
+	DOREPLIFETIME_CONDITION(AGameplayDebuggerReplicator, LastSelectedActorToDebug, COND_OwnerOnly);
+	DOREPLIFETIME_CONDITION(AGameplayDebuggerReplicator, ReplicatedObjects, COND_OwnerOnly);
+	DOREPLIFETIME_CONDITION(AGameplayDebuggerReplicator, Categories, COND_OwnerOnly);
+#endif
+}
+
+bool AGameplayDebuggerReplicator::ReplicateSubobjects(class UActorChannel *Channel, class FOutBunch *Bunch, FReplicationFlags *RepFlags)
+{
+	bool bWroteSomething = Super::ReplicateSubobjects(Channel, Bunch, RepFlags);
+#if ENABLED_GAMEPLAY_DEBUGGER
+
+	for (UGameplayDebuggerBaseObject* MyObject : ReplicatedObjects)
+	{
+		if (MyObject)
+		{
+			bWroteSomething |= Channel->ReplicateSubobject(MyObject, *Bunch, *RepFlags);
+		}
+	}
+#endif
+	return bWroteSomething;
+}
+
+
+bool AGameplayDebuggerReplicator::IsNetRelevantFor(const AActor* RealViewer, const AActor* ViewTarget, const FVector& SrcLocation) const
+{
+	return LocalPlayerOwner == RealViewer;
+}
+
+void AGameplayDebuggerReplicator::OnRep_ReplicatedCategories()
+{
+#if ENABLED_GAMEPLAY_DEBUGGER
+	MarkComponentsRenderStateDirty();
+#endif
+}
+
+void AGameplayDebuggerReplicator::BindKeyboardInput(class UInputComponent*& InInputComponent)
+{
+#if ENABLED_GAMEPLAY_DEBUGGER
+	if (!InInputComponent)
+	{
+		InInputComponent = NewObject<UInputComponent>(this, TEXT("AIDebugViewInInputComponent0"));
+		InInputComponent->RegisterComponent();
+
+		UGameplayDebuggerModuleSettings* Settings = UGameplayDebuggerModuleSettings::StaticClass()->GetDefaultObject<UGameplayDebuggerModuleSettings>();
+		FInputConfiguration& InputConfiguration = Settings->GetInputConfiguration();
+		InInputComponent->BindKey(InputConfiguration.SelectPlayer, IE_Pressed, this, &AGameplayDebuggerReplicator::OnSelectPlayer);
+		InInputComponent->BindKey(InputConfiguration.ZeroCategory, IE_Pressed, this, &AGameplayDebuggerReplicator::OnNumPadZero);
+		InInputComponent->BindKey(InputConfiguration.OneCategory, IE_Pressed, this, &AGameplayDebuggerReplicator::OnNumPadOne);
+		InInputComponent->BindKey(InputConfiguration.TwoCategory, IE_Pressed, this, &AGameplayDebuggerReplicator::OnNumPadTwo);
+		InInputComponent->BindKey(InputConfiguration.ThreeCategory, IE_Pressed, this, &AGameplayDebuggerReplicator::OnNumPadThree);
+		InInputComponent->BindKey(InputConfiguration.FourCategory, IE_Pressed, this, &AGameplayDebuggerReplicator::OnNumPadFour);
+		InInputComponent->BindKey(InputConfiguration.FiveCategory, IE_Pressed, this, &AGameplayDebuggerReplicator::OnNumPadFive);
+		InInputComponent->BindKey(InputConfiguration.SixCategory, IE_Pressed, this, &AGameplayDebuggerReplicator::OnNumPadSix);
+		InInputComponent->BindKey(InputConfiguration.SevenCategory, IE_Pressed, this, &AGameplayDebuggerReplicator::OnNumPadSeven);
+		InInputComponent->BindKey(InputConfiguration.EightCategory, IE_Pressed, this, &AGameplayDebuggerReplicator::OnNumPadEight);
+		InInputComponent->BindKey(InputConfiguration.NineCategory, IE_Pressed, this, &AGameplayDebuggerReplicator::OnNumPadNine);
+
+		InInputComponent->BindKey(InputConfiguration.NextMenuLine, IE_Released, this, &AGameplayDebuggerReplicator::OnNextMenuLine);
+		InInputComponent->BindKey(InputConfiguration.PreviousMenuLine, IE_Released, this, &AGameplayDebuggerReplicator::OnPreviousMenuLine);
+		InInputComponent->BindKey(InputConfiguration.DebugCamera, IE_Released, this, &AGameplayDebuggerReplicator::ToggleDebugCamera);
+		InInputComponent->BindKey(InputConfiguration.OnScreenDebugMessages, IE_Released, this, &AGameplayDebuggerReplicator::ToggleOnScreenDebugMessages);
+		InInputComponent->BindKey(InputConfiguration.GameHUD, IE_Released, this, &AGameplayDebuggerReplicator::ToggleGameHUD);
+
+		for (UGameplayDebuggerBaseObject* Obj : ReplicatedObjects)
+		{
+			if (Obj)
+			{
+				Obj->BindKeyboardImput(InInputComponent);
+			}
+		} 
+
+		if (LocalPlayerOwner)
+		{
+			LocalPlayerOwner->PushInputComponent(InInputComponent);
+		}
+	}
+#endif //ENABLED_GAMEPLAY_DEBUGGER 
+} 
+
+void AGameplayDebuggerReplicator::ToggleCategory(int32 CategoryIndex)
+{
+#if ENABLED_GAMEPLAY_DEBUGGER
+	const int32 FinalIndex = CurrentMenuLine * 10 + CategoryIndex;
+	if (Categories.IsValidIndex(FinalIndex))
+	{
+		ServerToggleCategory(Categories[FinalIndex].CategoryName); 
+	}
+
+	MarkComponentsRenderStateDirty();
+#endif
+} 
+
+void AGameplayDebuggerReplicator::OnNumPadZero() { ToggleCategory(0); }
+void AGameplayDebuggerReplicator::OnNumPadOne()  { ToggleCategory(1); }
+void AGameplayDebuggerReplicator::OnNumPadTwo() { ToggleCategory(2); }
+void AGameplayDebuggerReplicator::OnNumPadThree() { ToggleCategory(3); }
+void AGameplayDebuggerReplicator::OnNumPadFour() { ToggleCategory(4); }
+void AGameplayDebuggerReplicator::OnNumPadFive() { ToggleCategory(5); }
+void AGameplayDebuggerReplicator::OnNumPadSix() { ToggleCategory(6); }
+void AGameplayDebuggerReplicator::OnNumPadSeven() { ToggleCategory(7); }
+void AGameplayDebuggerReplicator::OnNumPadEight() { ToggleCategory(8); }
+void AGameplayDebuggerReplicator::OnNumPadNine() { ToggleCategory(9); }
+
+void AGameplayDebuggerReplicator::ToggleOnScreenDebugMessages()
+{
+#if ENABLED_GAMEPLAY_DEBUGGER
+	if (!IsActorTickEnabled() || LocalPlayerOwner == nullptr)
+	{
+		return;
+	}
+
+	if (GEngine)
+	{
+		GEngine->bEnableOnScreenDebugMessages = !GEngine->bEnableOnScreenDebugMessages;
+	}
+#endif
+}
+
+void AGameplayDebuggerReplicator::ToggleGameHUD()
+{
+#if ENABLED_GAMEPLAY_DEBUGGER
+	if (!IsActorTickEnabled() || LocalPlayerOwner == nullptr)
+	{
+		return;
+	}
+
+	if (AHUD* const GameHUD = LocalPlayerOwner->GetHUD())
+	{
+		GameHUD->bShowHUD = !GameHUD->bShowHUD;
+	}
+#endif
+}
+
+void AGameplayDebuggerReplicator::ToggleDebugCamera()
+{
+#if ENABLED_GAMEPLAY_DEBUGGER
+	if (LocalPlayerOwner == nullptr)
+	{
+		return;
+	}
+
+	if (DebugCameraController.IsValid() == false)
+	{
+		{
+			FActorSpawnParameters SpawnInfo;
+			SpawnInfo.SpawnCollisionHandlingOverride = ESpawnActorCollisionHandlingMethod::AlwaysSpawn;
+			SpawnInfo.Owner = LocalPlayerOwner->GetWorldSettings();
+			SpawnInfo.Instigator = LocalPlayerOwner->Instigator;
+			DebugCameraController = GetWorld()->SpawnActor<ADebugCameraController>(SpawnInfo);
+		}
+
+		if (DebugCameraController.IsValid())
+		{
+			// set up new controller
+			DebugCameraController->OnActivate(LocalPlayerOwner);
+
+			// then switch to it
+			LocalPlayerOwner->Player->SwitchController(DebugCameraController.Get());
+
+			FActorSpawnParameters SpawnInfo;
+			SpawnInfo.Owner = LocalPlayerOwner;
+			SpawnInfo.Instigator = LocalPlayerOwner->Instigator;
+			SpawnInfo.SpawnCollisionHandlingOverride = ESpawnActorCollisionHandlingMethod::AlwaysSpawn;
+			SpawnInfo.ObjectFlags |= RF_Transient;	// We never want these to save into a map
+			AGameplayDebuggerHUD* ProxyHUD = GetWorld()->SpawnActor<AGameplayDebuggerHUD>(SpawnInfo);
+			ProxyHUD->RedirectedHUD = LocalPlayerOwner->MyHUD;
+			DebugCameraController->MyHUD = ProxyHUD;
+			BindKeyboardInput(DebugCameraInputComponent);
+
+			const FInputActionKeyMapping& Mapping = UGameplayDebuggerModuleSettings::StaticClass()->GetDefaultObject<UGameplayDebuggerModuleSettings>()->ActivationAction;
+
+			FInputActionBinding& ActivationKeyPressed = DebugCameraInputComponent->BindAction(*Mapping.ActionName.ToString(), IE_Pressed, this, &AGameplayDebuggerReplicator::OnActivationKeyPressed);
+			ActivationKeyPressed.bConsumeInput = true;
+
+			FInputActionBinding& ActivationKeyReleased = DebugCameraInputComponent->BindAction(*Mapping.ActionName.ToString(), IE_Released, this, &AGameplayDebuggerReplicator::OnActivationKeyReleased);
+			ActivationKeyReleased.bConsumeInput = true;
+
+
+			DebugCameraController->ChangeState(NAME_Default);
+			DebugCameraController->ChangeState(NAME_Spectating);
+
+			if (LocalPlayerOwner)
+			{
+				LocalPlayerOwner->PopInputComponent(InputComponent);
+			}
+			DebugCameraController->PushInputComponent(DebugCameraInputComponent);
+			InputComponent = nullptr;
+		}
+	}
+	else
+	{
+		DebugCameraController->PopInputComponent(DebugCameraInputComponent);
+		DebugCameraInputComponent = nullptr;
+		DebugCameraController->OriginalPlayer->SwitchController(DebugCameraController->OriginalControllerRef);
+		DebugCameraController->OnDeactivate(DebugCameraController->OriginalControllerRef);
+		GetWorld()->DestroyActor(DebugCameraController.Get(), false, false);
+		DebugCameraController = nullptr;
+
+		if (InputComponent == nullptr)
+		{
+			BindKeyboardInput(InputComponent);
+			if (LocalPlayerOwner)
+			{
+				LocalPlayerOwner->PushInputComponent(InputComponent);
+			}
+		}
+
+	}
+
+#endif //ENABLED_GAMEPLAY_DEBUGGER
+}
+
+void AGameplayDebuggerReplicator::OnNextMenuLine()
+{
+#if ENABLED_GAMEPLAY_DEBUGGER
+	const int32 NumberOfLines = Categories.Num() / 10 + 1;
+	CurrentMenuLine = FMath::Clamp(CurrentMenuLine+1, 0, NumberOfLines-1);
+#endif
+}
+
+void AGameplayDebuggerReplicator::OnPreviousMenuLine()
+{
+#if ENABLED_GAMEPLAY_DEBUGGER
+	const int32 NumberOfLines = Categories.Num() / 10 + 1;
+	CurrentMenuLine = FMath::Clamp(CurrentMenuLine-1, 0, NumberOfLines - 1);
+#endif
+}
+
+bool AGameplayDebuggerReplicator::ServerToggleCategory_Validate(const FString& CategoryName)
+{
+#if ENABLED_GAMEPLAY_DEBUGGER
+	return true;
+#endif
+	return false;
+}
+
+void AGameplayDebuggerReplicator::ServerToggleCategory_Implementation(const FString& CategoryName)
+{
+#if ENABLED_GAMEPLAY_DEBUGGER
+	for (FGameplayDebuggerCategorySettings& CategorySettings : Categories)
+	{
+		if (CategorySettings.CategoryName == CategoryName)
+		{
+			CategorySettings.bPIE = !CategorySettings.bPIE;
+		}
+	}
+#endif
+}
+
+bool AGameplayDebuggerReplicator::IsCategoryEnabled(const FString& CategoryName)
+{
+#if ENABLED_GAMEPLAY_DEBUGGER
+	for (FGameplayDebuggerCategorySettings& CategorySettings : Categories)
+	{
+		if (CategorySettings.CategoryName == CategoryName)
+		{
+			return CategorySettings.bPIE;
+		}
+	}
+#endif
+	return false;
+}
+
+void AGameplayDebuggerReplicator::BeginPlay()
+{
+	Super::BeginPlay();
+
+#if ENABLED_GAMEPLAY_DEBUGGER
+
+	UGameplayDebuggerModuleSettings* Settings = UGameplayDebuggerModuleSettings::StaticClass()->GetDefaultObject<UGameplayDebuggerModuleSettings>();
+	if (Role == ROLE_Authority)
+	{
+		Settings->LoadAnyMissingClasses();
+
+		TArray<UClass *> Results = Settings->GetAllGameplayDebuggerClasses();
+		for (UClass* Class : Results)
+		{
+			UGameplayDebuggerBaseObject* BaseObject = NewObject<UGameplayDebuggerBaseObject>(this, Class, NAME_None, RF_Transient);
+			ReplicatedObjects.AddUnique(BaseObject);
+		}
+	}
+
+	// start with settings categories
+	Categories = Settings->GetCategories();
+
+	IConsoleVariable* cvarHighlightSelectedActor = IConsoleManager::Get().FindConsoleVariable(TEXT("ai.gd.HighlightSelectedActor"));
+	if (cvarHighlightSelectedActor)
+	{
+		cvarHighlightSelectedActor->Set(Settings->bHighlightSelectedActor);
+	}
+
+	const FInputActionKeyMapping& Mapping = Settings->ActivationAction;
+	ActivationKeyDisplayName = Mapping.Key.GetDisplayName().ToString();
+	ActivationKeyName = Mapping.Key.GetFName().ToString();
+
+	if (GetNetMode() != ENetMode::NM_DedicatedServer && LocalPlayerOwner && LocalPlayerOwner->InputComponent)
+	{
+		FInputActionBinding& ActivationKeyPressed = LocalPlayerOwner->InputComponent->BindAction(Mapping.ActionName, IE_Pressed, this, &AGameplayDebuggerReplicator::OnActivationKeyPressed);
+		ActivationKeyPressed.bConsumeInput = false;
+
+		FInputActionBinding& ActivationKeyReleased = LocalPlayerOwner->InputComponent->BindAction(Mapping.ActionName, IE_Released, this, &AGameplayDebuggerReplicator::OnActivationKeyReleased);
+		ActivationKeyReleased.bConsumeInput = false;
+
+		// register 'Game' show flag for Game or PIE, Simulate doesn't use it
+		UDebugDrawService::Register(TEXT("Game"), FDebugDrawDelegate::CreateUObject(this, &AGameplayDebuggerReplicator::DrawDebugDataDelegate));
+		if (GIsEditor)
+		{
+			// register 'DebugAI' show flag for Simulate, Game or PIE don't use DebugAI show flag but 'Game'
+			UDebugDrawService::Register(TEXT("DebugAI"), FDebugDrawDelegate::CreateUObject(this, &AGameplayDebuggerReplicator::OnDebugAIDelegate));
+		}
+	}
+#endif //ENABLED_GAMEPLAY_DEBUGGER
+}
+
+void AGameplayDebuggerReplicator::SetLocalPlayerOwner(APlayerController* PC)
+{
+#if ENABLED_GAMEPLAY_DEBUGGER
+	LocalPlayerOwner = PC;
+	if (GetNetMode() != ENetMode::NM_DedicatedServer && LocalPlayerOwner  && LocalPlayerOwner->InputComponent && ReplicatedObjects.Num() > 0)
+	{
+		const FInputActionKeyMapping& Mapping = UGameplayDebuggerModuleSettings::StaticClass()->GetDefaultObject<UGameplayDebuggerModuleSettings>()->ActivationAction;
+
+		FInputActionBinding& ActivationKeyPressed = LocalPlayerOwner->InputComponent->BindAction(Mapping.ActionName, IE_Pressed, this, &AGameplayDebuggerReplicator::OnActivationKeyPressed);
+		ActivationKeyPressed.bConsumeInput = false;
+
+		FInputActionBinding& ActivationKeyReleased = LocalPlayerOwner->InputComponent->BindAction(Mapping.ActionName, IE_Released, this, &AGameplayDebuggerReplicator::OnActivationKeyReleased);
+		ActivationKeyReleased.bConsumeInput = false;
+
+		// register 'Game' show flag for Game or PIE, Simulate doesn't use it
+		UDebugDrawService::Register(TEXT("Game"), FDebugDrawDelegate::CreateUObject(this, &AGameplayDebuggerReplicator::DrawDebugDataDelegate));
+		if (GIsEditor)
+		{
+			// register 'DebugAI' show flag for Simulate, Game or PIE don't use DebugAI show flag but 'Game'
+			UDebugDrawService::Register(TEXT("DebugAI"), FDebugDrawDelegate::CreateUObject(this, &AGameplayDebuggerReplicator::OnDebugAIDelegate));
+		}
+	}
+#endif
+}
+
+bool AGameplayDebuggerReplicator::ClientActivateGameplayDebugger_Validate(bool bActivate)
+{
+#if ENABLED_GAMEPLAY_DEBUGGER
+	return true;
+#endif
+	return false;
+}
+
+void AGameplayDebuggerReplicator::ClientActivateGameplayDebugger_Implementation(bool bActivate)
+{
+#if ENABLED_GAMEPLAY_DEBUGGER
+	SetActorTickEnabled(bActivate);
+	MarkComponentsRenderStateDirty();
+#endif
+}
+
+bool AGameplayDebuggerReplicator::ServerActivateGameplayDebugger_Validate(bool bActivate)
+{
+#if ENABLED_GAMEPLAY_DEBUGGER
+	return true;
+#endif
+	return false;
+}
+void AGameplayDebuggerReplicator::ServerActivateGameplayDebugger_Implementation(bool bActivate)
+{
+#if ENABLED_GAMEPLAY_DEBUGGER
+	bReplicates = !bActivate;
+	SetRemoteRoleForBackwardsCompat(ROLE_SimulatedProxy);
+	SetReplicates(bActivate);
+	SetActorTickEnabled(bActivate);
+#endif
+}
+
+
+void AGameplayDebuggerReplicator::BeginDestroy()
+{
+#if ENABLED_GAMEPLAY_DEBUGGER
+	if (GEngine)
+	{
+		GEngine->bEnableOnScreenDebugMessages = true;
+	}
+#endif
+	Super::BeginDestroy();
+}
+
+class UNetConnection* AGameplayDebuggerReplicator::GetNetConnection() const
+{
+#if ENABLED_GAMEPLAY_DEBUGGER
+	if (LocalPlayerOwner && LocalPlayerOwner->IsPendingKill() == false && IsPendingKill() == false)
+	{
+		return LocalPlayerOwner->GetNetConnection();
+	}
+#endif
+	return nullptr;
+}
+
+bool AGameplayDebuggerReplicator::ClientEnableTargetSelection_Validate(bool, APlayerController*)
+{
+#if ENABLED_GAMEPLAY_DEBUGGER
+	return true;
+#endif
+	return false;
+}
+
+void AGameplayDebuggerReplicator::ClientEnableTargetSelection_Implementation(bool bEnable, APlayerController* Context)
+{
+#if ENABLED_GAMEPLAY_DEBUGGER
+	bEnabledTargetSelection = bEnable;
+#endif
+}
+
+void AGameplayDebuggerReplicator::EnableDraw(bool bEnable)
+{
+#if ENABLED_GAMEPLAY_DEBUGGER
+	if (AHUD* const GameHUD = LocalPlayerOwner ? LocalPlayerOwner->GetHUD() : nullptr)
+	{
+		GameHUD->bShowHUD = bEnable ? false : true;
+	}
+	GEngine->bEnableOnScreenDebugMessages = bEnable ? false : true;
+#endif
+}
+
+void AGameplayDebuggerReplicator::ActivateFromCommand()
+{
+#if ENABLED_GAMEPLAY_DEBUGGER
+	OnActivationKeyPressed();
+	if (GetLocalPlayerOwner())
+	{
+		SelectTargetToDebug();
+	}
+	OnActivationKeyReleased();
+#endif
+}
+
+void AGameplayDebuggerReplicator::OnActivationKeyPressed()
+{
+#if ENABLED_GAMEPLAY_DEBUGGER
+	const bool bIsActive = IsActorTickEnabled();
+	bEnabledTargetSelection = true;
+	if (DebugCameraController.IsValid())
+	{
+		return;
+	}
+
+	ActivationKeyTime = 0;
+	bActivationKeyPressed = true;
+	if (!bIsActive)
+	{
+		SetActorTickEnabled(true);
+	}
+#endif
+}
+
+void AGameplayDebuggerReplicator::OnActivationKeyReleased()
+{
+#if ENABLED_GAMEPLAY_DEBUGGER
+	bEnabledTargetSelection = false;
+	if (DebugCameraController.IsValid())
+	{
+		return;
+	}
+
+	if (ActivationKeyTime < GameplayDebuggerHelpers::ActivationKeyTimePch && InputComponent)
+	{
+		if (LocalPlayerOwner)
+		{
+			LocalPlayerOwner->PopInputComponent(InputComponent);
+			InputComponent = nullptr;
+		}
+		ServerActivateGameplayDebugger(false);
+		ClientActivateGameplayDebugger(false);
+
+		GEngine->bEnableOnScreenDebugMessages = true;
+		if (AHUD* const GameHUD = LocalPlayerOwner ? LocalPlayerOwner->GetHUD() : nullptr)
+		{
+			GameHUD->bShowHUD = true;
+		}
+	}
+	else if (ActivationKeyTime < GameplayDebuggerHelpers::ActivationKeyTimePch && !InputComponent)
+	{
+		GEngine->bEnableOnScreenDebugMessages = false;
+		if (AHUD* const GameHUD = LocalPlayerOwner ? LocalPlayerOwner->GetHUD() : nullptr)
+		{
+			GameHUD->bShowHUD = false;
+		}
+		BindKeyboardInput(InputComponent);
+
+		ServerActivateGameplayDebugger(true);
+		ClientActivateGameplayDebugger(true);
+	}
+
+	bActivationKeyPressed = false;
+#endif
+}
+
+void AGameplayDebuggerReplicator::TickActor(float DeltaTime, enum ELevelTick TickType, FActorTickFunction& ThisTickFunction)
+{
+	Super::TickActor(DeltaTime, TickType, ThisTickFunction);
+
+#if ENABLED_GAMEPLAY_DEBUGGER
+	UWorld* World = GetWorld();
+	const ENetMode NetMode = GetNetMode();
+	if (!World)
+	{
+		// return without world
+		return;
+	}
+
+	UGameInstance* GameInstance = World->GetGameInstance();
+	if (!GameInstance || !World->IsGameWorld())
+	{
+		return;
+	}
+
+	if (NetMode != NM_DedicatedServer)
+	{
+		if (bActivationKeyPressed)
+		{
+			ActivationKeyTime += DeltaTime;
+			if (ActivationKeyTime >= GameplayDebuggerHelpers::ActivationKeyTimePch)
+			{
+				GEngine->bEnableOnScreenDebugMessages = false;
+				if (AHUD* const GameHUD = LocalPlayerOwner ? LocalPlayerOwner->GetHUD() : nullptr)
+				{
+					GameHUD->bShowHUD = false;
+				}
+				BindKeyboardInput(InputComponent);
+
+				ServerActivateGameplayDebugger(true);
+				ClientActivateGameplayDebugger(true);
+				bActivationKeyPressed = false;
+			}
+		}
+
+		if (bEnabledTargetSelection)
+		{
+			if (GetLocalPlayerOwner())
+			{
+				SelectTargetToDebug();
+			}
+		}
+
+		bool bMarkComponentsAsRenderStateDirty = false;
+		for (UGameplayDebuggerBaseObject* Obj : ReplicatedObjects)
+		{
+			if (Obj && Obj->IsRenderStateDirty())
+			{
+				if (!bMarkComponentsAsRenderStateDirty)
+				{
+					MarkComponentsRenderStateDirty();
+				}
+				bMarkComponentsAsRenderStateDirty = true;
+				Obj->CleanRenderStateDirtyFlag();
+			}
+		}
+	}
+
+	if (NetMode < NM_Client && LocalPlayerOwner)
+	{
+		TMap<FString, TArray<UGameplayDebuggerBaseObject*> > CategoryToClasses;
+		for (UGameplayDebuggerBaseObject* Obj : ReplicatedObjects)
+		{
+			if (Obj)
+			{
+				FString Category = Obj->GetCategoryName();
+				if (IsCategoryEnabled(Category))
+				{
+					CategoryToClasses.FindOrAdd(Category).Add(Obj);
+				}
+			}
+		}
+
+		for (auto It(CategoryToClasses.CreateIterator()); It; ++It)
+		{
+			TArray<UGameplayDebuggerBaseObject*>& CurrentObjects = It.Value();
+			for (UGameplayDebuggerBaseObject* Obj : CurrentObjects)
+			{
+				Obj->CollectDataToReplicateOnServer(LocalPlayerOwner, LastSelectedActorToDebug);
+			}
+		}
+	}
+
+#endif
+}
+
+void AGameplayDebuggerReplicator::OnSelectPlayer()
+{
+#if ENABLED_GAMEPLAY_DEBUGGER
+	APlayerController* MyPC = DebugCameraController.IsValid() ? DebugCameraController.Get() : GetLocalPlayerOwner();
+	APawn* MyPawn = MyPC ? MyPC->GetPawnOrSpectator() : nullptr;
+	if (MyPawn)
+	{
+		ServerSetActorToDebug(MyPawn);
+	}
+#endif
+}
+
+void AGameplayDebuggerReplicator::SelectTargetToDebug()
+{
+#if ENABLED_GAMEPLAY_DEBUGGER
+	APlayerController* MyPC = DebugCameraController.IsValid() ? DebugCameraController.Get() : GetLocalPlayerOwner();
+
+	if (MyPC)
+	{
+		float bestAim = 0;
+		FVector CamLocation;
+		FRotator CamRotation;
+		check(MyPC->PlayerCameraManager != nullptr);
+		MyPC->PlayerCameraManager->GetCameraViewPoint(CamLocation, CamRotation);
+		FVector FireDir = CamRotation.Vector();
+		UWorld* World = MyPC->GetWorld();
+		check(World);
+
+		APawn* PossibleTarget = nullptr;
+		for (FConstPawnIterator Iterator = World->GetPawnIterator(); Iterator; ++Iterator)
+		{
+			APawn* NewTarget = *Iterator;
+			if (NewTarget == nullptr || NewTarget == MyPC->GetPawn()
+				|| (NewTarget->PlayerState != nullptr && NewTarget->PlayerState->bIsABot == false)
+				|| NewTarget->GetActorEnableCollision() == false)
+			{
+				continue;
+			}
+
+			if (NewTarget != MyPC->GetPawn())
+			{
+				// look for best controlled pawn target
+				const FVector AimDir = NewTarget->GetActorLocation() - CamLocation;
+				float FireDist = AimDir.SizeSquared();
+				// only find targets which are < 25000 units away
+				if (FireDist < 625000000.f)
+				{
+					FireDist = FMath::Sqrt(FireDist);
+					float newAim = FVector::DotProduct(FireDir, AimDir);
+					newAim = newAim / FireDist;
+					if (newAim > bestAim)
+					{
+						PossibleTarget = NewTarget;
+						bestAim = newAim;
+					}
+				}
+			}
+		}
+
+		if (PossibleTarget != nullptr && PossibleTarget != LastSelectedActorToDebug)
+		{
+			ServerSetActorToDebug(Cast<AActor>(PossibleTarget));
+		}
+	}
+#endif //ENABLED_GAMEPLAY_DEBUGGER
+}
+
+bool AGameplayDebuggerReplicator::ServerSetActorToDebug_Validate(AActor* InActor)
+{
+#if ENABLED_GAMEPLAY_DEBUGGER
+	return true;
+#endif
+	return false;
+}
+
+void AGameplayDebuggerReplicator::ServerSetActorToDebug_Implementation(AActor* InActor)
+{
+#if ENABLED_GAMEPLAY_DEBUGGER
+	if (LastSelectedActorToDebug != InActor)
+	{
+		LastSelectedActorToDebug = InActor;
+	}
+#endif
+}
+
+void AGameplayDebuggerReplicator::OnDebugAIDelegate(class UCanvas* Canvas, class APlayerController* PC)
+{
+#if WITH_EDITOR && ENABLED_GAMEPLAY_DEBUGGER
+	if (!GIsEditor)
+	{
+		return;
+	}
+
+	if (!LocalPlayerOwner)
+	{
+		return;
+	}
+
+	UEditorEngine* EEngine = Cast<UEditorEngine>(GEngine);
+	if (!EEngine || !EEngine->bIsSimulatingInEditor)
+	{
+		return;
+	}
+
+	if (!Canvas || !Canvas->SceneView || Canvas->SceneView->bIsGameView == false)
+	{
+		return;
+	}
+
+	FEngineShowFlags EngineShowFlags = Canvas && Canvas->SceneView && Canvas->SceneView->Family ? Canvas->SceneView->Family->EngineShowFlags : FEngineShowFlags(GIsEditor ? EShowFlagInitMode::ESFIM_Editor : EShowFlagInitMode::ESFIM_Game);
+	if (!EngineShowFlags.DebugAI)
+	{
+		return;
+	}
+
+	EnableDraw(true);
+	UWorld* World = GetWorld();
+	if (World && Role == ROLE_Authority)
+	{
+		if (IsActorTickEnabled() == false)
+		{
+			SetActorTickEnabled(true);
+		}
+
+		// looks like Simulate in UE4 Editor - let's find selected Pawn to debug
+		AActor* SelectedActor = nullptr;
+		for (FSelectionIterator It = EEngine->GetSelectedActorIterator(); It; ++It)
+		{
+			SelectedActor = Cast<APawn>(*It); //we only work with pawns for now
+			if (SelectedActor)
+			{
+				break;
+			}
+		}
+
+		if (LastSelectedActorToDebug != SelectedActor)
+		{
+			MarkComponentsRenderStateDirty();
+		}
+
+		ServerSetActorToDebug(SelectedActor);
+		DrawDebugData(Canvas, PC, true);
+	}
+#endif
+}
+
+void AGameplayDebuggerReplicator::DrawDebugDataDelegate(class UCanvas* Canvas, class APlayerController* PC)
+{
+#if ENABLED_GAMEPLAY_DEBUGGER
+	if (GetWorld() == nullptr || IsPendingKill() || Canvas == nullptr || Canvas->IsPendingKill())
+	{
+		return;
+	}
+
+	if (!LocalPlayerOwner || !IsActorTickEnabled())
+	{
+		return;
+	}
+
+	if (Canvas->SceneView != nullptr && !Canvas->SceneView->bIsGameView)
+	{
+		return;
+	}
+
+	if (GetWorld()->bPlayersOnly && Role == ROLE_Authority)
+	{
+		for (FConstPawnIterator Iterator = GetWorld()->GetPawnIterator(); Iterator; ++Iterator)
+		{
+			AActor* NewTarget = Cast<AActor>(*Iterator);
+			if (NewTarget->IsSelected() && GetSelectedActorToDebug() != NewTarget)
+			{
+				ServerSetActorToDebug(NewTarget);
+			}
+		}
+	}
+	DrawDebugData(Canvas, PC);
+#endif
+}
+
+void AGameplayDebuggerReplicator::DrawDebugData(class UCanvas* Canvas, class APlayerController* PC, bool bHideMenu)
+{
+#if ENABLED_GAMEPLAY_DEBUGGER
+	if (!LocalPlayerOwner && IsActorTickEnabled())
+	{
+		return;
+	}
+
+	const bool bAllowToDraw = Canvas && Canvas->SceneView && (Canvas->SceneView->ViewActor == LocalPlayerOwner->AcknowledgedPawn || Canvas->SceneView->ViewActor == LocalPlayerOwner->GetPawnOrSpectator());
+	if (!bAllowToDraw)
+	{
+		// check for spectator debug camera during debug camera
+		if (DebugCameraController.IsValid() == false || Canvas->SceneView->ViewActor->GetInstigatorController() != DebugCameraController.Get())
+		{
+			return;
+		}
+	}
+
+	const float DebugInfoStartX = UGameplayDebuggerModuleSettings::StaticClass()->GetDefaultObject<UGameplayDebuggerModuleSettings>()->DebugInfoStart.X;
+	const float DebugInfoStartY = UGameplayDebuggerModuleSettings::StaticClass()->GetDefaultObject<UGameplayDebuggerModuleSettings>()->DebugInfoStart.Y;
+	const FVector SelectedActorLoc = LastSelectedActorToDebug ? LastSelectedActorToDebug->GetActorLocation() + FVector(0, 0, LastSelectedActorToDebug->GetSimpleCollisionHalfHeight()) : DebugTools::InvalidLocation;
+	
+	UGameplayDebuggerHelper::FPrintContext DefaultContext(GEngine->GetSmallFont(), Canvas, DebugInfoStartX, DebugInfoStartY);
+	DefaultContext.FontRenderInfo.bEnableShadow = true;
+	const bool bDrawFullData = SelectedActorLoc != DebugTools::InvalidLocation;
+	const FVector ScreenLoc = SelectedActorLoc != DebugTools::InvalidLocation ? UGameplayDebuggerHelper::ProjectLocation(DefaultContext, SelectedActorLoc) : FVector::ZeroVector;
+	UGameplayDebuggerHelper::FPrintContext OverHeadContext(GEngine->GetSmallFont(), Canvas, ScreenLoc.X, ScreenLoc.Y);
+
+	UGameplayDebuggerHelper::SetOverHeadContext(OverHeadContext);
+	UGameplayDebuggerHelper::SetDefaultContext(DefaultContext);
+
+	if (DefaultContext.Canvas != nullptr)
+	{
+		float XL, YL;
+		const FString ToolName = FString::Printf(TEXT("Gameplay Debugger [Timestamp: %05.03f]"), GetWorld()->TimeSeconds);
+		UGameplayDebuggerHelper::CalulateStringSize(DefaultContext, nullptr, ToolName, XL, YL);
+		UGameplayDebuggerHelper::PrintString(DefaultContext, FColorList::White, ToolName, DefaultContext.Canvas->ClipX / 2.0f - XL / 2.0f, 0);
+	}
+
+	if (!bHideMenu)
+	{
+		DrawMenu(DefaultContext, OverHeadContext);
+	}
+
+	TMap<FString, TArray<UGameplayDebuggerBaseObject*> > CategoryToClasses;
+	for (UGameplayDebuggerBaseObject* Obj : ReplicatedObjects)
+	{
+		if (Obj)
+		{
+			FString Category = Obj->GetCategoryName();
+			CategoryToClasses.FindOrAdd(Category).Add(Obj);
+		}
+	}
+	CategoryToClasses.KeySort(TLess<FString>());
+
+	for (auto It(CategoryToClasses.CreateIterator()); It; ++It)
+	{
+		const FGameplayDebuggerCategorySettings* Element = Categories.FindByPredicate([&](const FGameplayDebuggerCategorySettings& C){ return It.Key() == C.CategoryName; });
+		if (Element == nullptr || Element->bPIE == false)
+		{
+			continue;
+		}
+
+		UGameplayDebuggerHelper::PrintString(UGameplayDebuggerHelper::GetDefaultContext(), FString::Printf(TEXT("\n{R=0,G=255,B=0,A=255}%s\n"), *It.Key()));
+		TArray<UGameplayDebuggerBaseObject*>& CurrentObjects = It.Value();
+		for (UGameplayDebuggerBaseObject* Obj : CurrentObjects)
+		{
+			Obj->DrawCollectedData(LocalPlayerOwner, LastSelectedActorToDebug);
+		}
+	}
+
+	const IConsoleVariable* cvarHighlightSelectedActor = IConsoleManager::Get().FindConsoleVariable(TEXT("ai.gd.HighlightSelectedActor"));
+	const bool bHighlightSelectedActor = !cvarHighlightSelectedActor || cvarHighlightSelectedActor->GetInt();
+	if (LastSelectedActorToDebug && bHighlightSelectedActor)
+	{
+		FBox ComponentsBoundingBox = LastSelectedActorToDebug->GetComponentsBoundingBox(false);
+		DrawDebugBox(GetWorld(), ComponentsBoundingBox.GetCenter(), ComponentsBoundingBox.GetExtent(), FColor::Red, false);
+		DrawDebugSolidBox(GetWorld(), ComponentsBoundingBox.GetCenter(), ComponentsBoundingBox.GetExtent(), FColor::Red.WithAlpha(25));
+	}
+#endif
+}
+
+void AGameplayDebuggerReplicator::DrawMenu(UGameplayDebuggerBaseObject::FPrintContext& DefaultContext, UGameplayDebuggerBaseObject::FPrintContext& OverHeadContext)
+{
+#if ENABLED_GAMEPLAY_DEBUGGER
+	const float OldX = DefaultContext.CursorX;
+	const float OldY = DefaultContext.CursorY;
+
+	const UGameplayDebuggerModuleSettings* GameplayDebuggerSettings = UGameplayDebuggerModuleSettings::StaticClass()->GetDefaultObject<UGameplayDebuggerModuleSettings>();
+	const float MenuStartX = GameplayDebuggerSettings->MenuStart.X;
+	const float MenuStartY = GameplayDebuggerSettings->MenuStart.Y;
+
+	if (DefaultContext.Canvas != nullptr)
+	{
+		UFont* OldFont = DefaultContext.Font.Get();
+		DefaultContext.Font = GEngine->GetMediumFont();
+
+		TArray<float> CategoriesWidth;
+		CategoriesWidth.AddZeroed(Categories.Num());
+
+		APlayerController* const MyPC = Cast<APlayerController>(LocalPlayerOwner);
+
+		FString KeyDesc = ActivationKeyName != ActivationKeyDisplayName ? FString::Printf(TEXT("(%s key)"), *ActivationKeyName) : TEXT("");
+		FString HeaderDesc;
+		if (UGameplayDebuggerModuleSettings::StaticClass()->GetDefaultObject<UGameplayDebuggerModuleSettings>()->UseAlternateKeys())
+		{
+			HeaderDesc = FString::Printf(TEXT("{LightBlue}Tap {white}%s %s {LightBlue}to close, {white}Alt + number {LightBlue}to toggle categories, use {white}+ {LightBlue}or {white}- {LightBlue}to select active line with categories"), *ActivationKeyDisplayName, *KeyDesc);
+		}
+		else
+		{
+			HeaderDesc = FString::Printf(TEXT("{LightBlue}Use {white}%s %s{LightBlue} to close, {white}Numpad numbers {LightBlue}to toggle categories or use {white}+ {LightBlue}and {white}- {LightBlue}to select active line with categories"), *ActivationKeyDisplayName, *KeyDesc);
+		}
+
+		FString DefaultMenu;
+		{
+			DefaultMenu += FString::Printf(TEXT(" {white}[Tab]: %s%s  "), DebugCameraController.IsValid() ? TEXT("{Green}") : TEXT("{LightBlue}"), TEXT("Debug Camera"));
+		}
+		{
+			DefaultMenu += FString::Printf(TEXT(" {white}[Ctrl+Tab]: %s%s  "), GEngine && GEngine->bEnableOnScreenDebugMessages ? TEXT("{Green}") : TEXT("{LightBlue}"), TEXT("DebugMessages"));
+		}
+		{
+			const AHUD* GameHUD = MyPC ? MyPC->GetHUD() : nullptr;
+			DefaultMenu += FString::Printf(TEXT(" {white}[Ctrl+Tilde]: %s%s  "), GameHUD && GameHUD->bShowHUD ? TEXT("{Green}") : TEXT("{LightBlue}"), TEXT("GameHUD"));
+		}
+		{
+			KeyDesc = UGameplayDebuggerModuleSettings::StaticClass()->GetDefaultObject<UGameplayDebuggerModuleSettings>()->GetInputConfiguration().SelectPlayer.Key.GetFName().ToString();
+			const bool SelectedPlayer = MyPC && MyPC->GetPawnOrSpectator() && LastSelectedActorToDebug == MyPC->GetPawnOrSpectator();
+			DefaultMenu += FString::Printf(TEXT(" {white}[%s]: %sPlayer  "), *KeyDesc, SelectedPlayer ? TEXT("{Green}") : TEXT("{LightBlue}"));
+		}
+
+		FString MenuString = HeaderDesc + TEXT("\n") + DefaultMenu + TEXT("\n");
+		const int32 NumberOfLines = Categories.Num() / 10 + 1;
+		for (int32 LineIndex = 0; LineIndex < NumberOfLines; ++LineIndex)
+		{
+			const int32 ItemsInLine = FMath::Clamp(Categories.Num() - LineIndex * 10, 0, 10);
+			for (int32 MenuItemIndex = 0; MenuItemIndex < ItemsInLine; ++MenuItemIndex)
+			{
+				const int32 i = LineIndex * 10 + MenuItemIndex;
+
+				const bool bIsActive = IsCategoryEnabled(Categories[i].CategoryName);
+				const bool bIsDisabled = LastSelectedActorToDebug == nullptr;
+				const FString CategoryColor = bIsDisabled ? (bIsActive ? TEXT("DarkGreen") : TEXT("LightGrey")) : (bIsActive ? TEXT("Green") : TEXT("LightBlue"));
+				if (CurrentMenuLine == LineIndex)
+				{
+					MenuString += FString::Printf(TEXT("   {white}%d:{%s}%s "), MenuItemIndex, *CategoryColor, *Categories[i].CategoryName);
+				}
+				else
+				{
+					MenuString += FString::Printf(TEXT("   {%s}%s "), *CategoryColor, *Categories[i].CategoryName);
+				}
+			}
+			MenuString += TEXT(" \n");
+		}
+		float MenuWidth, MenuHeight;
+		UGameplayDebuggerBaseObject::CalulateStringSize(DefaultContext, nullptr, MenuString, MenuWidth, MenuHeight);
+		MaxMenuWidth = FMath::Max(MaxMenuWidth, MenuWidth+4);
+		MenuHeight += 4;
+
+		FCanvasTileItem TileItem(FVector2D(0, 0), GWhiteTexture, FVector2D(MaxMenuWidth, MenuHeight), FColor(0, 0, 0, 30));
+		TileItem.BlendMode = SE_BLEND_Translucent;
+		UGameplayDebuggerBaseObject::DrawItem(DefaultContext, TileItem, MenuStartX, MenuStartY);
+
+		float CharWidth, CharHeight;
+		UGameplayDebuggerBaseObject::CalulateStringSize(DefaultContext, nullptr, TEXT("X"), CharWidth, CharHeight);
+		TileItem = FCanvasTileItem(FVector2D(0, 0), GWhiteTexture, FVector2D(MaxMenuWidth, CharHeight), FColor(0, 0, 0, 50));
+		TileItem.BlendMode = SE_BLEND_Translucent;
+		UGameplayDebuggerBaseObject::DrawItem(DefaultContext, TileItem, MenuStartX, MenuStartY + 2 + (CurrentMenuLine+2) * CharHeight);
+		UGameplayDebuggerBaseObject::PrintString(DefaultContext, MenuString, MenuStartX + 2, MenuStartY + 2);
+
+		DefaultContext.Font = OldFont;
+	}
+
+	if (LastSelectedActorToDebug == nullptr)
+	{
+		UGameplayDebuggerBaseObject::PrintString(DefaultContext, "{red}No Pawn selected - waiting for data to replicate from server. {green}Press and hold ' to select Pawn \n");
+	}
+
+	if (DebugCameraController.IsValid())
+	{
+		ADebugCameraController* DebugCamController = DebugCameraController.Get();
+		if (DebugCamController != nullptr)
+		{
+			FVector const CamLoc = DebugCamController->PlayerCameraManager->GetCameraLocation();
+			FRotator const CamRot = DebugCamController->PlayerCameraManager->GetCameraRotation();
+
+			FString HitString;
+			FCollisionQueryParams TraceParams(NAME_None, true, this);
+			FHitResult Hit;
+			bool bHit = GetWorld()->LineTraceSingleByChannel(Hit, CamLoc, CamRot.Vector() * 100000.f + CamLoc, ECC_Pawn, TraceParams);
+			if (bHit && Hit.GetActor() != nullptr)
+			{
+				HitString = FString::Printf(TEXT("{white}Under cursor: {yellow}'%s'"), *Hit.GetActor()->GetName());
+				DrawDebugLine(GetWorld(), Hit.Location, Hit.Location + Hit.Normal*60.f, FColor::White, /*bPersistentLines =*/ false, /*LifeTime =*/ -1.f, /*DepthPriority =*/ 0, /*Thickness =*/ 3.f);
+			}
+			else
+			{
+				HitString = FString::Printf(TEXT("Not actor under cursor"));
+			}
+
+			UGameplayDebuggerBaseObject::PrintString(DefaultContext, FColor::White, HitString/*, MenuStartX, MenuStartY + 40*/);
+		}
+	}
+
+
+	DefaultContext.CursorX = OldX;
+	DefaultContext.CursorY = OldY;
+#endif //ENABLED_GAMEPLAY_DEBUGGER
+}
+
+void AGameplayDebuggerReplicator::DebugNextPawn(UClass* CompareClass, APawn* CurrentPawn)
+{
+#if ENABLED_GAMEPLAY_DEBUGGER
+	APawn* LastSeen = nullptr;
+	APawn* FirstSeen = nullptr;
+	// Search through the list looking for the next one of this type
+	for (FConstPawnIterator It = GetWorld()->GetPawnIterator(); It; It++)
+	{
+		APawn* IterPawn = *It;
+		if (IterPawn->IsA(CompareClass))
+		{
+			if (LastSeen == CurrentPawn)
+			{
+				ServerSetActorToDebug(IterPawn);
+				return;
+			}
+			LastSeen = IterPawn;
+			if (FirstSeen == nullptr)
+			{
+				FirstSeen = IterPawn;
+			}
+		}
+	}
+	// See if we need to wrap around the list
+	if (FirstSeen != nullptr)
+	{
+		ServerSetActorToDebug(FirstSeen);
+	}
+#endif
+}
+
+void AGameplayDebuggerReplicator::DebugPrevPawn(UClass* CompareClass, APawn* CurrentPawn)
+{
+#if ENABLED_GAMEPLAY_DEBUGGER
+	APawn* LastSeen = nullptr;
+	APawn* PrevSeen = nullptr;
+	APawn* FirstSeen = nullptr;
+	// Search through the list looking for the prev one of this type
+	for (FConstPawnIterator It = GetWorld()->GetPawnIterator(); It; It++)
+	{
+		APawn* IterPawn = *It;
+		if (IterPawn->IsA(CompareClass))
+		{
+			if (LastSeen == CurrentPawn && FirstSeen != CurrentPawn)
+			{
+				ServerSetActorToDebug(PrevSeen);
+				return;
+			}
+			PrevSeen = LastSeen;
+			LastSeen = IterPawn;
+			if (FirstSeen == nullptr)
+			{
+				FirstSeen = IterPawn;
+				if (CurrentPawn == nullptr)
+				{
+					ServerSetActorToDebug(FirstSeen);
+					return;
+				}
+			}
+		}
+	}
+	// Wrap from the beginning to the end
+	if (FirstSeen == CurrentPawn)
+	{
+		ServerSetActorToDebug(LastSeen);
+	}
+	// Handle getting the previous to the end
+	else if (LastSeen == CurrentPawn)
+	{
+		ServerSetActorToDebug(PrevSeen);
+	}
+#endif
+}
+
+//////////////////////////////////////////////////////////////////////////
+//----------------------------------------------------------------------//
+// Composite Scene proxy
+//----------------------------------------------------------------------//
+class FDebugRenderSceneCompositeProxy : public FDebugRenderSceneProxy
+{
+public:
+	FDebugRenderSceneCompositeProxy(const UPrimitiveComponent* InComponent)
+		: FDebugRenderSceneProxy(InComponent)
+	{
+	}
+
+	virtual ~FDebugRenderSceneCompositeProxy()
+	{
+		for (int32 Index = 0; Index < ChildProxies.Num(); ++Index)
+		{
+			delete ChildProxies[Index];
+		}
+	}
+
+	virtual void DrawStaticElements(FStaticPrimitiveDrawInterface* PDI) override
+	{
+		for (int32 Index = 0; Index < ChildProxies.Num(); ++Index)
+		{
+			ChildProxies[Index]->DrawStaticElements(PDI);
+		}
+	}
+
+	virtual void GetDynamicMeshElements(const TArray<const FSceneView*>& Views, const FSceneViewFamily& ViewFamily, uint32 VisibilityMap, FMeshElementCollector& Collector) const override
+	{
+		for (int32 Index = 0; Index < ChildProxies.Num(); ++Index)
+		{
+			ChildProxies[Index]->GetDynamicMeshElements(Views, ViewFamily, VisibilityMap, Collector);
+		}
+	}
+
+	virtual void RegisterDebugDrawDelgate() override
+	{
+		for (int32 Index = 0; Index < ChildProxies.Num(); ++Index)
+		{
+			ChildProxies[Index]->RegisterDebugDrawDelgate();
+		}
+	}
+
+	virtual void UnregisterDebugDrawDelgate() override
+	{
+		for (int32 Index = 0; Index < ChildProxies.Num(); ++Index)
+		{
+			ChildProxies[Index]->UnregisterDebugDrawDelgate();
+		}
+	}
+
+	virtual FPrimitiveViewRelevance GetViewRelevance(const FSceneView* View) const override
+	{
+		FPrimitiveViewRelevance Result;
+		for (int32 Index = 0; Index < ChildProxies.Num(); ++Index)
+		{
+			Result |= ChildProxies[Index]->GetViewRelevance(View);
+		}
+		return Result;
+	}
+
+	virtual uint32 GetMemoryFootprint(void) const override
+	{
+		return sizeof(*this) + GetAllocatedSize();
+	}
+
+	uint32 GetAllocatedSize(void) const
+	{
+		uint32 Size = 0;
+		for (int32 Index = 0; Index < ChildProxies.Num(); ++Index)
+		{
+			ChildProxies[Index]->GetMemoryFootprint();
+		}
+
+		return Size + ChildProxies.GetAllocatedSize();
+	}
+
+	void AddChild(FDebugRenderSceneProxy* NewChild)
+	{
+		ChildProxies.AddUnique(NewChild);
+	}
+
+protected:
+	TArray<FDebugRenderSceneProxy*> ChildProxies;
+};
+
+//////////////////////////////////////////////////////////////////////////
+//----------------------------------------------------------------------//
+// UGameplayDebuggerDrawComponent
+//----------------------------------------------------------------------//
+UGameplayDebuggerDrawComponent::UGameplayDebuggerDrawComponent(const FObjectInitializer& ObjectInitializer) : Super(ObjectInitializer) { }
+
+FPrimitiveSceneProxy* UGameplayDebuggerDrawComponent::CreateSceneProxy()
+{
+	FDebugRenderSceneCompositeProxy* CompositeProxy = nullptr;
+#if ENABLED_GAMEPLAY_DEBUGGER
+	AGameplayDebuggerReplicator* Replicator = Cast<AGameplayDebuggerReplicator>(GetOwner());
+	if (!Replicator || Replicator->IsActorTickEnabled() == false || !World || World->GetNetMode() == NM_DedicatedServer || IsPendingKill())
+	{
+		return nullptr;
+	}
+
+	CompositeProxy = new FDebugRenderSceneCompositeProxy(this);
+
+	for (UGameplayDebuggerBaseObject* Obj : Replicator->ReplicatedObjects)
+	{
+		if (Obj && Replicator->IsCategoryEnabled(Obj->GetCategoryName()))
+		{
+			FDebugRenderSceneProxy* ChildSceneProxy = Obj->CreateSceneProxy(this, World, Replicator->GetSelectedActorToDebug());
+			if (ChildSceneProxy)
+			{
+				CompositeProxy->AddChild(ChildSceneProxy);
+			}
+		}
+	}
+
+#endif //ENABLED_GAMEPLAY_DEBUGGER
+	return CompositeProxy;
+}
+
+FBoxSphereBounds UGameplayDebuggerDrawComponent::CalcBounds(const FTransform& LocalToWorld) const
+{
+	//we know nothing about 
+	FBox MyBounds(FVector(-HALF_WORLD_MAX1, -HALF_WORLD_MAX1, -HALF_WORLD_MAX1), FVector(HALF_WORLD_MAX1, HALF_WORLD_MAX1, HALF_WORLD_MAX1));
+	return MyBounds;
+}
+
+void UGameplayDebuggerDrawComponent::CreateRenderState_Concurrent()
+{
+	Super::CreateRenderState_Concurrent();
+
+#if WITH_EDITOR
+	if (SceneProxy)
+	{
+		static_cast<FDebugRenderSceneCompositeProxy*>(SceneProxy)->RegisterDebugDrawDelgate();
+	}
+#endif
+}
+
+void UGameplayDebuggerDrawComponent::DestroyRenderState_Concurrent()
+{
+#if WITH_EDITOR
+	if (SceneProxy)
+	{
+		static_cast<FDebugRenderSceneCompositeProxy*>(SceneProxy)->UnregisterDebugDrawDelgate();
+	}
+#endif
+
+	Super::DestroyRenderState_Concurrent();
+}
+
+//////////////////////////////////////////////////////////////////////////
+//----------------------------------------------------------------------//
+// AGaneplayDebuggerHUD
+//----------------------------------------------------------------------//
+AGameplayDebuggerHUD::AGameplayDebuggerHUD(const FObjectInitializer& ObjectInitializer) : Super(ObjectInitializer)
+{
+}
+
+void AGameplayDebuggerHUD::PostRender()
+{
+#if ENABLED_GAMEPLAY_DEBUGGER
+	if (RedirectedHUD.IsValid())
+	{
+		RedirectedHUD->SetCanvas(Canvas, DebugCanvas);
+		RedirectedHUD->PostRender();
+	}
+#endif
+}
+
diff --git a/Engine/Plugins/Developer/GameplayDebugger/Source/GameplayDebuggerModule/Private/GameplayDebuggerReplicator.h b/Engine/Plugins/Developer/GameplayDebugger/Source/GameplayDebuggerModule/Private/GameplayDebuggerReplicator.h
new file mode 100644
--- /dev/null
+++ b/Engine/Plugins/Developer/GameplayDebugger/Source/GameplayDebuggerModule/Private/GameplayDebuggerReplicator.h
@@ -0,0 +1,174 @@
+// Copyright 1998-2015 Epic Games, Inc. All Rights Reserved.
+
+#pragma once
+#include "GameplayDebuggerModuleSettings.h"
+#include "Debug/GameplayDebuggerBaseObject.h"
+#include "GameFramework/HUD.h"
+#include "GameplayDebuggerReplicator.generated.h"
+/**
+*	Transient actor used to communicate between server and client, mostly for RPC
+*/
+
+class UGameplayDebuggingComponent;
+class UGameplayDebuggerBaseObject;
+class AGameplayDebuggingHUDComponent;
+class ADebugCameraController;
+
+UCLASS(NotBlueprintable, Transient, hidecategories = Actor, notplaceable)
+class AGameplayDebuggerReplicator : public AActor
+{
+	GENERATED_UCLASS_BODY()
+
+	UPROPERTY(Transient, Replicated)
+	APlayerController* LocalPlayerOwner;
+
+	UPROPERTY(Transient, Replicated)
+	AActor*	LastSelectedActorToDebug;
+
+	UPROPERTY(Transient, Replicated)
+	TArray<class UGameplayDebuggerBaseObject*> ReplicatedObjects;
+
+	UPROPERTY(Transient, ReplicatedUsing = OnRep_ReplicatedCategories)
+	TArray<FGameplayDebuggerCategorySettings> Categories;
+
+	UFUNCTION()
+	virtual void OnRep_ReplicatedCategories();
+
+	UPROPERTY()
+	class UTexture2D* DefaultTexture_Red;
+
+	UPROPERTY()
+	class UTexture2D* DefaultTexture_Green;
+
+	UPROPERTY(Transient)
+	class UInputComponent* DebugCameraInputComponent;
+
+	UFUNCTION(Reliable, Client, WithValidation)
+	void ClientEnableTargetSelection(bool bEnable, APlayerController* Context);
+
+	UFUNCTION(Reliable, Server, WithValidation)
+	void ServerToggleCategory(const FString& CategoryName);
+
+	bool IsCategoryEnabled(const FString& CategoryName);
+
+	virtual class UNetConnection* GetNetConnection() const override;
+
+	virtual bool IsNetRelevantFor(const AActor* RealViewer, const AActor* ViewTarget, const FVector& SrcLocation) const override;
+
+	virtual void BeginPlay() override;
+
+	virtual void TickActor(float DeltaTime, enum ELevelTick TickType, FActorTickFunction& ThisTickFunction) override;
+
+	virtual void BeginDestroy() override;
+
+	virtual bool ReplicateSubobjects(class UActorChannel *Channel, class FOutBunch *Bunch, FReplicationFlags *RepFlags) override;
+
+	void SelectTargetToDebug();
+
+	void EnableDraw(bool bEnable);
+
+	void SetLocalPlayerOwner(APlayerController* PC);
+	APlayerController* GetLocalPlayerOwner() { return LocalPlayerOwner; }
+
+	FORCEINLINE AActor* GetSelectedActorToDebug() { return LastSelectedActorToDebug; }
+
+	UFUNCTION(Reliable, Server, WithValidation)
+	void ServerSetActorToDebug(AActor* InActor);
+
+	UFUNCTION(Reliable, Client, WithValidation)
+	void ClientActivateGameplayDebugger(bool bActivate);
+
+	UFUNCTION(Reliable, Server, WithValidation)
+	void ServerActivateGameplayDebugger(bool bActivate);
+
+	/**
+	* Iterates through the pawn list to find the next pawn of the specified type to debug
+	*/
+	void DebugNextPawn(UClass* CompareClass, APawn* CurrentPawn = nullptr);
+
+	/**
+	* Iterates through the pawn list to find the previous pawn of the specified type to debug
+	*/
+	void DebugPrevPawn(UClass* CompareClass, APawn* CurrentPawn = nullptr);
+
+	void OnActivationKeyPressed();
+
+	void ActivateFromCommand();
+
+protected:
+	void BindKeyboardInput(class UInputComponent*& InInputComponent);
+	void OnDebugAIDelegate(class UCanvas* Canvas, class APlayerController* PC);
+	void DrawDebugDataDelegate(class UCanvas* Canvas, class APlayerController* PC);
+	void DrawDebugData(class UCanvas* Canvas, class APlayerController* PC, bool bHideMenu = false);
+	void DrawMenu(UGameplayDebuggerBaseObject::FPrintContext& DefaultContext, UGameplayDebuggerBaseObject::FPrintContext& OverHeadContext);
+
+	void OnSelectPlayer();
+	void OnNumPadZero();
+	void OnNumPadOne();
+	void OnNumPadTwo();
+	void OnNumPadThree();
+	void OnNumPadFour();
+	void OnNumPadFive();
+	void OnNumPadSix();
+	void OnNumPadSeven();
+	void OnNumPadEight();
+	void OnNumPadNine();
+
+	void OnNextMenuLine();
+	void OnPreviousMenuLine();
+
+	void OnActivationKeyReleased();
+
+	void ToggleCategory(int32 CategoryIndex);
+	void ToggleDebugCamera();
+	void ToggleOnScreenDebugMessages();
+	void ToggleGameHUD();
+
+private:
+	TWeakObjectPtr<ADebugCameraController> DebugCameraController;
+
+	uint32 bEnabledTargetSelection : 1;
+	uint8 bActivationKeyPressed : 1;
+	int32 CurrentMenuLine;
+	float PlayerControllersUpdateDelay;
+	float ActivationKeyTime;
+
+	float MaxMenuWidth;
+	FString ActivationKeyDisplayName;
+	FString ActivationKeyName;
+	FDelegateHandle OnSelectionChangedHandle;
+};
+
+UCLASS()
+class UGameplayDebuggerDrawComponent : public UPrimitiveComponent
+{
+	GENERATED_UCLASS_BODY()
+public:
+	virtual FPrimitiveSceneProxy* CreateSceneProxy() override;
+	virtual FBoxSphereBounds CalcBounds(const FTransform &LocalToWorld) const override;
+	virtual void CreateRenderState_Concurrent() override;
+	virtual void DestroyRenderState_Concurrent() override;
+};
+
+UCLASS(NotBlueprintable, Transient, NotBlueprintType, hidedropdown)
+class AGameplayDebuggerHUD : public AHUD
+{
+	GENERATED_UCLASS_BODY()
+
+	FFontRenderInfo TextRenderInfo;
+
+	//~ Begin AActor Interface
+	virtual void PostRender() override;
+	//~ End AActor Interface
+
+	TWeakObjectPtr<AHUD> RedirectedHUD;
+};
+
+UCLASS(NotBlueprintable, Transient, hidecategories = Actor, notplaceable)
+class AGameplayDebuggerServerHelper : public AActor
+{
+	GENERATED_UCLASS_BODY()
+
+	virtual void TickActor(float DeltaTime, enum ELevelTick TickType, FActorTickFunction& ThisTickFunction) override;
+};
+
diff --git a/Engine/Plugins/Developer/GameplayDebugger/Source/GameplayDebuggerModule/Private/GameplayDebuggerServerHelper.cpp b/Engine/Plugins/Developer/GameplayDebugger/Source/GameplayDebuggerModule/Private/GameplayDebuggerServerHelper.cpp
new file mode 100644
--- /dev/null
+++ b/Engine/Plugins/Developer/GameplayDebugger/Source/GameplayDebuggerModule/Private/GameplayDebuggerServerHelper.cpp
@@ -0,0 +1,56 @@
+// Copyright 1998-2015 Epic Games, Inc. All Rights Reserved.
+
+#include "GameplayDebuggerPluginPrivatePCH.h"
+#include "GameplayDebuggerReplicator.h"
+
+AGameplayDebuggerServerHelper::AGameplayDebuggerServerHelper(const FObjectInitializer& ObjectInitializer) : Super(ObjectInitializer)
+{
+#if ENABLED_GAMEPLAY_DEBUGGER
+	const bool bCanTick = true;
+#else
+	const bool bCanTick = false;
+#endif
+
+	PrimaryActorTick.bCanEverTick = bCanTick;
+	PrimaryActorTick.bAllowTickOnDedicatedServer = bCanTick;
+	PrimaryActorTick.bTickEvenWhenPaused = bCanTick;
+	PrimaryActorTick.bStartWithTickEnabled = bCanTick;
+	PrimaryActorTick.TickInterval = 2;
+
+#if WITH_EDITOR
+	SetIsTemporarilyHiddenInEditor(true);
+#endif
+
+#if WITH_EDITORONLY_DATA
+	bHiddenEdLevel = true;
+	bHiddenEdLayer = true;
+	bHiddenEd = true;
+	bEditable = false;
+#endif
+
+}
+
+void AGameplayDebuggerServerHelper::TickActor(float DeltaTime, enum ELevelTick TickType, FActorTickFunction& ThisTickFunction)
+{
+	Super::TickActor(DeltaTime, TickType, ThisTickFunction);
+
+#if ENABLED_GAMEPLAY_DEBUGGER
+	const UWorld* World = GetWorld();
+	if (!World || World->IsGameWorld() == false)
+	{
+		// return without world
+		return;
+	}
+
+	for (FConstPlayerControllerIterator Iterator = World->GetPlayerControllerIterator(); Iterator; ++Iterator)
+	{
+		APlayerController* PC = *Iterator;
+		if (PC)
+		{
+			IGameplayDebuggerPlugin& Debugger = IGameplayDebuggerPlugin::Get();
+			Debugger.CreateGameplayDebuggerForPlayerController(PC);
+			PrimaryActorTick.TickInterval = 5;
+		}
+	}
+#endif
+}
diff --git a/Engine/Plugins/Developer/GameplayDebugger/Source/GameplayDebuggerModule/Public/GameplayDebuggerBlueprintBaseObject.h b/Engine/Plugins/Developer/GameplayDebugger/Source/GameplayDebuggerModule/Public/GameplayDebuggerBlueprintBaseObject.h
new file mode 100644
--- /dev/null
+++ b/Engine/Plugins/Developer/GameplayDebugger/Source/GameplayDebuggerModule/Public/GameplayDebuggerBlueprintBaseObject.h
@@ -0,0 +1,32 @@
+// Copyright 1998-2015 Epic Games, Inc. All Rights Reserved.
+
+#pragma once
+#include "Debug/GameplayDebuggerBaseObject.h"
+#include "GameplayDebuggerBlueprintBaseObject.generated.h"
+
+UCLASS(Blueprintable, Transient, Abstract, EarlyAccessPreview)
+class UGameplayDebuggerBlueprintBaseObject : public UGameplayDebuggerBaseObject
+{
+	GENERATED_UCLASS_BODY()
+
+	UPROPERTY(EditDefaultsOnly, BlueprintReadOnly, AssetRegistrySearchable, Category = "GameplayDebugger", meta = (DisplayName = "CategoryName"))
+	FString GameplayDebuggerCategoryName;
+
+	UFUNCTION(BlueprintImplementableEvent, Category = "GameplayDebugger", meta = (DisplayName = "CollectDataToReplicateOnServer"))
+	void CollectDataToReplicateOnServerEvent(APlayerController* PlayerController, AActor *SelectedActor);
+
+	UFUNCTION(BlueprintImplementableEvent, Category = "GameplayDebugger", meta = (DisplayName = "DrawCollectedDataOnClient"))
+	void DrawCollectedDataEvent(APlayerController* PlayerController, AActor* SelectedActor = nullptr);
+
+	virtual FString GetCategoryName() const override;
+
+	virtual void DrawCollectedData(APlayerController* PlayerController, AActor* SelectedActor = nullptr) override;
+
+#if WITH_EDITOR
+	// Begin UObject Interface
+	virtual void PostEditChangeProperty(struct FPropertyChangedEvent& PropertyChangedEvent) override;
+#endif
+
+protected:
+	virtual void CollectDataToReplicate(APlayerController* PlayerController, AActor *SelectedActor) override { CollectDataToReplicateOnServerEvent(PlayerController, SelectedActor); }
+};
diff --git a/Engine/Plugins/Developer/GameplayDebugger/Source/GameplayDebuggerModule/Public/IGameplayDebuggerPlugin.h b/Engine/Plugins/Developer/GameplayDebugger/Source/GameplayDebuggerModule/Public/IGameplayDebuggerPlugin.h
new file mode 100644
--- /dev/null
+++ b/Engine/Plugins/Developer/GameplayDebugger/Source/GameplayDebuggerModule/Public/IGameplayDebuggerPlugin.h
@@ -0,0 +1,40 @@
+// Copyright 1998-2015 Epic Games, Inc. All Rights Reserved.
+
+#pragma once
+
+#include "ModuleManager.h"
+
+/**
+ * The public interface to this module.  In most cases, this interface is only public to sibling modules 
+ * within this plugin.
+ */
+class IGameplayDebuggerPlugin : public IModuleInterface
+{
+
+public:
+
+	/**
+	 * Singleton-like access to this module's interface.  This is just for convenience!
+	 * Beware of calling this during the shutdown phase, though.  Your module might have been unloaded already.
+	 *
+	 * @return Returns singleton instance, loading the module on demand if needed
+	 */
+	static inline IGameplayDebuggerPlugin& Get()
+	{
+		return FModuleManager::LoadModuleChecked< IGameplayDebuggerPlugin >("GameplayDebuggerModule");
+	}
+
+	/**
+	 * Checks to see if this module is loaded and ready.  It is only valid to call Get() if IsAvailable() returns true.
+	 *
+	 * @return True if the module is loaded and ready to use
+	 */
+	static inline bool IsAvailable()
+	{
+		return FModuleManager::Get().IsModuleLoaded( "GameplayDebugger" );
+	}
+
+	//
+	virtual AActor* CreateGameplayDebuggerForPlayerController(class APlayerController* PlayerController) = 0;
+};
+
diff --git a/Engine/Plugins/Runtime/ProceduralMeshComponent/Source/Private/KismetProceduralMeshLibrary.cpp b/Engine/Plugins/Runtime/ProceduralMeshComponent/Source/Private/KismetProceduralMeshLibrary.cpp
--- a/Engine/Plugins/Runtime/ProceduralMeshComponent/Source/Private/KismetProceduralMeshLibrary.cpp
+++ b/Engine/Plugins/Runtime/ProceduralMeshComponent/Source/Private/KismetProceduralMeshLibrary.cpp
@@ -3,6 +3,8 @@
 #include "ProceduralMeshComponentPluginPrivatePCH.h"
 #include "KismetProceduralMeshLibrary.h"
 
+DECLARE_CYCLE_STAT(TEXT("Update Collision"), STAT_ProcMesh_CalcTangents, STATGROUP_ProceduralMesh);
+
 
 UKismetProceduralMeshLibrary::UKismetProceduralMeshLibrary(const FObjectInitializer& ObjectInitializer)
 : Super(ObjectInitializer)
@@ -158,6 +160,8 @@ void FindVertOverlaps(int32 TestVertIndex, const TArray<FVector>& Verts, TArray<
 
 void UKismetProceduralMeshLibrary::CalculateTangentsForMesh(const TArray<FVector>& Vertices, const TArray<int32>& Triangles, const TArray<FVector2D>& UVs, TArray<FVector>& Normals, TArray<FProcMeshTangent>& Tangents)
 {
+	SCOPE_CYCLE_COUNTER(STAT_ProcMesh_CalcTangents);
+
 	if (Vertices.Num() == 0)
 	{
 		return;
diff --git a/Engine/Plugins/Runtime/ProceduralMeshComponent/Source/Private/ProceduralMeshComponent.cpp b/Engine/Plugins/Runtime/ProceduralMeshComponent/Source/Private/ProceduralMeshComponent.cpp
--- a/Engine/Plugins/Runtime/ProceduralMeshComponent/Source/Private/ProceduralMeshComponent.cpp
+++ b/Engine/Plugins/Runtime/ProceduralMeshComponent/Source/Private/ProceduralMeshComponent.cpp
@@ -4,6 +4,14 @@
 #include "ProceduralMeshComponent.h"
 #include "DynamicMeshBuilder.h"
 
+DECLARE_CYCLE_STAT(TEXT("Create ProcMesh Proxy"), STAT_ProcMesh_CreateSceneProxy, STATGROUP_ProceduralMesh);
+DECLARE_CYCLE_STAT(TEXT("Create Mesh Section"), STAT_ProcMesh_CreateMeshSection, STATGROUP_ProceduralMesh);
+DECLARE_CYCLE_STAT(TEXT("UpdateSection GT"), STAT_ProcMesh_UpdateSectionGT, STATGROUP_ProceduralMesh);
+DECLARE_CYCLE_STAT(TEXT("UpdateSection RT"), STAT_ProcMesh_UpdateSectionRT, STATGROUP_ProceduralMesh);
+DECLARE_CYCLE_STAT(TEXT("Get ProcMesh Elements"), STAT_ProcMesh_GetMeshElements, STATGROUP_ProceduralMesh);
+DECLARE_CYCLE_STAT(TEXT("Update Collision"), STAT_ProcMesh_UpdateCollision, STATGROUP_ProceduralMesh);
+
+
 /** Resource array to pass  */
 class FProcMeshVertexResourceArray : public FResourceArrayInterface
 {
@@ -226,6 +234,8 @@ public:
 	/** Called on render thread to assign new dynamic data */
 	void UpdateSection_RenderThread(FProcMeshSectionUpdateData* SectionData)
 	{
+		SCOPE_CYCLE_COUNTER(STAT_ProcMesh_UpdateSectionRT);
+
 		check(IsInRenderingThread());
 
 		// Check we have data 
@@ -271,6 +281,9 @@ public:
 
 	virtual void GetDynamicMeshElements(const TArray<const FSceneView*>& Views, const FSceneViewFamily& ViewFamily, uint32 VisibilityMap, FMeshElementCollector& Collector) const override
 	{
+		SCOPE_CYCLE_COUNTER(STAT_ProcMesh_GetMeshElements);
+
+
 		// Set up wireframe material (if needed)
 		const bool bWireframe = AllowDebugViewmodes() && ViewFamily.EngineShowFlags.Wireframe;
 
@@ -378,6 +391,8 @@ UProceduralMeshComponent::UProceduralMeshComponent(const FObjectInitializer& Obj
 
 void UProceduralMeshComponent::CreateMeshSection(int32 SectionIndex, const TArray<FVector>& Vertices, const TArray<int32>& Triangles, const TArray<FVector>& Normals, const TArray<FVector2D>& UV0, const TArray<FColor>& VertexColors, const TArray<FProcMeshTangent>& Tangents, bool bCreateCollision)
 {
+	SCOPE_CYCLE_COUNTER(STAT_ProcMesh_CreateMeshSection);
+
 	// Ensure sections array is long enough
 	ProcMeshSections.SetNum(SectionIndex + 1, false);
 
@@ -423,6 +438,8 @@ void UProceduralMeshComponent::CreateMeshSection(int32 SectionIndex, const TArra
 
 void UProceduralMeshComponent::UpdateMeshSection(int32 SectionIndex, const TArray<FVector>& Vertices, const TArray<FVector>& Normals, const TArray<FVector2D>& UV0, const TArray<FColor>& VertexColors, const TArray<FProcMeshTangent>& Tangents)
 {
+	SCOPE_CYCLE_COUNTER(STAT_ProcMesh_UpdateSectionGT);
+
 	if(SectionIndex < ProcMeshSections.Num())
 	{
 		FProcMeshSection& Section = ProcMeshSections[SectionIndex];
@@ -581,6 +598,8 @@ void UProceduralMeshComponent::UpdateLocalBounds()
 
 FPrimitiveSceneProxy* UProceduralMeshComponent::CreateSceneProxy()
 {
+	SCOPE_CYCLE_COUNTER(STAT_ProcMesh_CreateSceneProxy);
+
 	return new FProceduralMeshSceneProxy(this);
 }
 
@@ -664,6 +683,8 @@ void UProceduralMeshComponent::CreateProcMeshBodySetup()
 
 void UProceduralMeshComponent::UpdateCollision()
 {
+	SCOPE_CYCLE_COUNTER(STAT_ProcMesh_UpdateCollision);
+
 	bool bCreatePhysState = false; // Should we create physics state at the end of this function?
 
 	// If its created, shut it down now
diff --git a/Engine/Plugins/Runtime/ProceduralMeshComponent/Source/Private/ProceduralMeshComponentPluginPrivatePCH.h b/Engine/Plugins/Runtime/ProceduralMeshComponent/Source/Private/ProceduralMeshComponentPluginPrivatePCH.h
--- a/Engine/Plugins/Runtime/ProceduralMeshComponent/Source/Private/ProceduralMeshComponentPluginPrivatePCH.h
+++ b/Engine/Plugins/Runtime/ProceduralMeshComponent/Source/Private/ProceduralMeshComponentPluginPrivatePCH.h
@@ -8,4 +8,8 @@
 #include "IProceduralMeshComponentPlugin.h"
 
 #include "ProceduralMeshComponent.h"
-#include "PhysicsEngine/BodySetup.h"
\ No newline at end of file
+#include "PhysicsEngine/BodySetup.h"
+
+#pragma once
+
+DECLARE_STATS_GROUP(TEXT("ProceduralMesh"), STATGROUP_ProceduralMesh, STATCAT_Advanced);
diff --git a/Engine/Source/Developer/AssetTools/Private/AssetFixUpRedirectors.cpp b/Engine/Source/Developer/AssetTools/Private/AssetFixUpRedirectors.cpp
--- a/Engine/Source/Developer/AssetTools/Private/AssetFixUpRedirectors.cpp
+++ b/Engine/Source/Developer/AssetTools/Private/AssetFixUpRedirectors.cpp
@@ -187,7 +187,7 @@ void FAssetFixUpRedirectors::LoadReferencingPackages(TArray<FRedirectorRefs>& Re
 
 			if ( Package )
 			{
-				if ( Package->PackageFlags & PKG_CompiledIn )
+				if ( Package->HasAnyPackageFlags(PKG_CompiledIn) )
 				{
 					// This is a script reference
 					RedirectorRefs.bRedirectorValidForFixup = false;
diff --git a/Engine/Source/Developer/BlueprintNativeCodeGen/Private/BlueprintNativeCodeGenManifest.cpp b/Engine/Source/Developer/BlueprintNativeCodeGen/Private/BlueprintNativeCodeGenManifest.cpp
--- a/Engine/Source/Developer/BlueprintNativeCodeGen/Private/BlueprintNativeCodeGenManifest.cpp
+++ b/Engine/Source/Developer/BlueprintNativeCodeGen/Private/BlueprintNativeCodeGenManifest.cpp
@@ -201,7 +201,7 @@ static bool BlueprintNativeCodeGenManifestImpl::GatherModuleDependencies(const U
 			}
 
 			// we want only native packages, ones that are not editor-only
-			if ((DependentPackage->PackageFlags & (PKG_CompiledIn | PKG_EditorOnly)) == PKG_CompiledIn)
+			if ((DependentPackage->GetPackageFlags() & (PKG_CompiledIn | PKG_EditorOnly)) == PKG_CompiledIn)
 			{
 				DependenciesOut.AddUnique(DependentPackage);// PkgImport.ObjectName.ToString());
 			}
diff --git a/Engine/Source/Developer/GameplayDebugger/Classes/GameplayDebuggerSettings.h b/Engine/Source/Developer/GameplayDebugger/Classes/GameplayDebuggerSettings.h
--- a/Engine/Source/Developer/GameplayDebugger/Classes/GameplayDebuggerSettings.h
+++ b/Engine/Source/Developer/GameplayDebugger/Classes/GameplayDebuggerSettings.h
@@ -63,8 +63,6 @@ public:
 		return Settings;
 	}
 
-	bool UseAlternateKeys() { return bUseAlternateKeys; }
-
 	const FGDTCustomViewNames& GetCustomViewNames() { return CustomViewNames; }
 
 #if ADD_LEVEL_EDITOR_EXTENSIONS
@@ -113,9 +111,6 @@ protected:
 	UPROPERTY(EditAnywhere, config, Category = "GameplayDebugger")
 	bool GameView5;
 
-	UPROPERTY(EditAnywhere, config, Category = "GameplayDebugger")
-	bool bUseAlternateKeys;
-
 	// Holds an event delegate that is executed when a setting has changed.
 	FSettingChangedEvent SettingChangedEvent;
 };
diff --git a/Engine/Source/Developer/GameplayDebugger/Classes/GameplayDebuggingComponent.h b/Engine/Source/Developer/GameplayDebugger/Classes/GameplayDebuggingComponent.h
--- a/Engine/Source/Developer/GameplayDebugger/Classes/GameplayDebuggingComponent.h
+++ b/Engine/Source/Developer/GameplayDebugger/Classes/GameplayDebuggingComponent.h
@@ -10,6 +10,7 @@
 #include "EnvironmentQuery/EQSQueryResultSourceInterface.h"
 #include "EnvironmentQuery/EnvQueryDebugHelpers.h"
 #include "Debug/DebugDrawService.h"
+#include "Debug/GameplayDebuggerBaseObject.h"
 #include "GameplayDebuggingComponent.generated.h"
 
 #define WITH_EQS 1
@@ -167,7 +168,7 @@ class GAMEPLAYDEBUGGER_API UGameplayDebuggingComponent : public UPrimitiveCompon
 	FVector SensingComponentLocation;
 
 	UPROPERTY(Replicated)
-	TArray<FDrawDebugShapeElement> PerceptionShapeElements;
+	TArray<FGameplayDebuggerShapeElement> PerceptionShapeElements;
 	/** End - Perception System */
 
 	UFUNCTION()
@@ -254,16 +255,13 @@ public:
 	virtual FBoxSphereBounds CalcBounds(const FTransform &LocalToWorld) const override;
 	virtual void CreateRenderState_Concurrent() override;
 	virtual void DestroyRenderState_Concurrent() override;
-
-protected:
 	void SelectTargetToDebug();
 
-	//APlayerController* PlayerOwner;
+protected:
 #if WITH_RECAST
 	ARecastNavMesh* GetNavData();
 #endif
 
-protected:
 	virtual void CollectPathData();
 	virtual void CollectBasicData();
 	virtual void CollectBehaviorTreeData();
diff --git a/Engine/Source/Developer/GameplayDebugger/Classes/GameplayDebuggingControllerComponent.h b/Engine/Source/Developer/GameplayDebugger/Classes/GameplayDebuggingControllerComponent.h
--- a/Engine/Source/Developer/GameplayDebugger/Classes/GameplayDebuggingControllerComponent.h
+++ b/Engine/Source/Developer/GameplayDebugger/Classes/GameplayDebuggingControllerComponent.h
@@ -6,6 +6,7 @@
 
 #pragma once
 #include "TimerManager.h"
+#include "InputCore.h"
 #include "GameplayDebuggingTypes.h"
 #include "GameplayDebugger.h"
 #include "GameFramework/HUD.h"
@@ -97,15 +98,60 @@ protected:
 	/** Handle for efficient management of UpdateNavMesh timer */
 	FTimerHandle TimerHandle_UpdateNavMeshTimer;
 
-	FInputChord ActivationKey;
 	const float KeyPressActivationTime;
 	float ControlKeyPressedTime;
 	float PlayersComponentRequestTime;
 	uint32 bToolActivated : 1;
 	uint32 bWaitingForOwnersComponent : 1;
+
+public:
+	UPROPERTY(config)
+	FInputChord ActivationKey;
+
+	UPROPERTY(config)
+	FInputChord CategoryZeroBind;
+
+	UPROPERTY(config)
+	FInputChord CategoryOneBind;
+
+	UPROPERTY(config)
+	FInputChord CategoryTwoBind;
+
+	UPROPERTY(config)
+	FInputChord CategoryThreeBind;
+
+	UPROPERTY(config)
+	FInputChord CategoryFourBind;
+
+	UPROPERTY(config)
+	FInputChord CategoryFiveBind;
+
+	UPROPERTY(config)
+	FInputChord CategorySixBind;
+
+	UPROPERTY(config)
+	FInputChord CategorySevenBind;
+
+	UPROPERTY(config)
+	FInputChord CategoryEightBind;
+
+	UPROPERTY(config)
+	FInputChord CategoryNineBind;
+
+	UPROPERTY(config)
+	FInputChord CycleDetailsViewBind;
+
+	UPROPERTY(config)
+	FInputChord DebugCameraBind;
+
+	UPROPERTY(config)
+	FInputChord OnScreenDebugMessagesBind;
+
+	UPROPERTY(config)
+	FInputChord GameHUDBind;
 };
 
-UCLASS()
+UCLASS(NotBlueprintable, Transient, NotBlueprintType, hidedropdown, hidecategories = Actor, notplaceable)
 class AGaneplayDebuggerProxyHUD : public AHUD
 {
 	GENERATED_UCLASS_BODY()
diff --git a/Engine/Source/Developer/GameplayDebugger/Classes/GameplayDebuggingTypes.h b/Engine/Source/Developer/GameplayDebugger/Classes/GameplayDebuggingTypes.h
--- a/Engine/Source/Developer/GameplayDebugger/Classes/GameplayDebuggingTypes.h
+++ b/Engine/Source/Developer/GameplayDebugger/Classes/GameplayDebuggingTypes.h
@@ -54,5 +54,3 @@ GAMEPLAYDEBUGGER_API
 FGameplayDebuggerSettings  GameplayDebuggerSettings(class AGameplayDebuggingReplicator* Replicator = NULL);
 
 DECLARE_MULTICAST_DELEGATE(FOnCycleDetailsView);
-
-DECLARE_LOG_CATEGORY_EXTERN(LogGameplayDebugger, Warning, All);
diff --git a/Engine/Source/Developer/GameplayDebugger/Private/GameplayDebugger.cpp b/Engine/Source/Developer/GameplayDebugger/Private/GameplayDebugger.cpp
--- a/Engine/Source/Developer/GameplayDebugger/Private/GameplayDebugger.cpp
+++ b/Engine/Source/Developer/GameplayDebugger/Private/GameplayDebugger.cpp
@@ -15,8 +15,6 @@
 #include "LevelEditor.h"
 #endif // WITH_EDITOR
 
-DEFINE_LOG_CATEGORY(LogGameplayDebugger);
-
 #define LOCTEXT_NAMESPACE "FGameplayDebugger"
 FGameplayDebuggerSettings GameplayDebuggerSettings(class AGameplayDebuggingReplicator* Replicator)
 {
@@ -69,6 +67,9 @@ IMPLEMENT_MODULE(FGameplayDebugger, GameplayDebugger)
 // This code will execute after your module is loaded into memory (but after global variables are initialized, of course.)
 void FGameplayDebugger::StartupModule()
 { 
+	//EMIT_CUSTOM_WARNING("/Engine/Source/Developer/GameplayDebugger module is deprecated and it's going to be removed with next UE4 version. Please use GameplayDebuggerPluging instead.");
+	UE_LOG(LogGameplayDebugger, Warning, TEXT("/Engine/Source/Developer/GameplayDebugger module is deprecated and it's going to be removed with next UE4 version. Please use GameplayDebuggerPluging instead."));
+
 #if !(UE_BUILD_SHIPPING || UE_BUILD_TEST)
 	if (GEngine)
 	{
@@ -616,9 +617,12 @@ bool FGameplayDebugger::Exec(UWorld* InWorld, const TCHAR* Cmd, FOutputDevice& A
 				bHandled = true;
 			}
 		}
-		else if (!Replicator->IsToolCreated())
+		else
 		{
-			Replicator->CreateTool();
+			if (Replicator->IsToolCreated() == false)
+			{
+				Replicator->CreateTool();
+			}
 			Replicator->EnableTool();
 			bHandled = true;
 		}
diff --git a/Engine/Source/Developer/GameplayDebugger/Private/GameplayDebuggerPrivate.h b/Engine/Source/Developer/GameplayDebugger/Private/GameplayDebuggerPrivate.h
--- a/Engine/Source/Developer/GameplayDebugger/Private/GameplayDebuggerPrivate.h
+++ b/Engine/Source/Developer/GameplayDebugger/Private/GameplayDebuggerPrivate.h
@@ -31,6 +31,7 @@
 
 //////////////////////////////////////////////////////////////////////////
 
+#include "Debug/GameplayDebuggerBaseObject.h"
 #include "GameplayDebuggingTypes.h"
 #include "GameplayDebugger.h"
 #include "GameplayDebuggingReplicator.h"
diff --git a/Engine/Source/Developer/GameplayDebugger/Private/GameplayDebuggingComponent.cpp b/Engine/Source/Developer/GameplayDebugger/Private/GameplayDebuggingComponent.cpp
--- a/Engine/Source/Developer/GameplayDebugger/Private/GameplayDebuggingComponent.cpp
+++ b/Engine/Source/Developer/GameplayDebugger/Private/GameplayDebuggingComponent.cpp
@@ -1586,11 +1586,15 @@ void UGameplayDebuggingComponent::CollectPerceptionData()
 		if (BTAI)
 		{
 			const UAIPerceptionComponent* PerceptionComponent = BTAI->GetAIPerceptionComponent();
+			if (PerceptionComponent == nullptr)
+			{
+				PerceptionComponent = MyPawn->FindComponentByClass<UAIPerceptionComponent>();
+			}
 			if (PerceptionComponent)
 			{
 				TArray<FString> PerceptionTexts;
 				PerceptionShapeElements.Reset();
-				PerceptionComponent->GetDebugData(PerceptionTexts, PerceptionShapeElements);
+				PerceptionComponent->GrabGameplayDebuggerData(PerceptionTexts, PerceptionShapeElements);
 
 				DistanceFromPlayer = DistanceFromSensor = -1;
 
diff --git a/Engine/Source/Developer/GameplayDebugger/Private/GameplayDebuggingControllerComponent.cpp b/Engine/Source/Developer/GameplayDebugger/Private/GameplayDebuggingControllerComponent.cpp
--- a/Engine/Source/Developer/GameplayDebugger/Private/GameplayDebuggingControllerComponent.cpp
+++ b/Engine/Source/Developer/GameplayDebugger/Private/GameplayDebuggingControllerComponent.cpp
@@ -45,7 +45,22 @@ UGameplayDebuggingControllerComponent::UGameplayDebuggingControllerComponent(con
 	bWaitingForOwnersComponent = false;
 
 	ControlKeyPressedTime = 0;
+
 	ActivationKey = FInputChord(EKeys::Apostrophe, false, false, false, false);
+	CategoryZeroBind = FInputChord(EKeys::NumPadZero, false, false, false, false);
+	CategoryOneBind = FInputChord(EKeys::NumPadOne, false, false, false, false);
+	CategoryTwoBind = FInputChord(EKeys::NumPadTwo, false, false, false, false);
+	CategoryThreeBind = FInputChord(EKeys::NumPadThree, false, false, false, false);
+	CategoryFourBind = FInputChord(EKeys::NumPadFour, false, false, false, false);
+	CategoryFiveBind = FInputChord(EKeys::NumPadFive, false, false, false, false);
+	CategorySixBind = FInputChord(EKeys::NumPadSix, false, false, false, false);
+	CategorySevenBind = FInputChord(EKeys::NumPadSeven, false, false, false, false);
+	CategoryEightBind = FInputChord(EKeys::NumPadEight, false, false, false, false);
+	CategoryNineBind = FInputChord(EKeys::NumPadNine, false, false, false, false);
+	CycleDetailsViewBind = FInputChord(EKeys::Add, false, false, false, false);
+	DebugCameraBind = FInputChord(EKeys::Tab, false, false, false, false);
+	OnScreenDebugMessagesBind = FInputChord(EKeys::Tab, false, true, false, false);
+	GameHUDBind = FInputChord(EKeys::Tilde, false, true, false, false);
 #endif //!(UE_BUILD_SHIPPING || UE_BUILD_TEST)
 }
 
@@ -276,37 +291,20 @@ void UGameplayDebuggingControllerComponent::BindAIDebugViewKeys(class UInputComp
 		InputComponent = NewObject<UInputComponent>(GetOwner(), TEXT("AIDebugViewInputComponent0"));
 		InputComponent->RegisterComponent();
 
-		if (UGameplayDebuggerSettings::StaticClass()->GetDefaultObject<UGameplayDebuggerSettings>()->UseAlternateKeys())
-		{
-			InputComponent->BindKey(FInputChord(EKeys::Zero, false, false, true, false), IE_Pressed, this, &UGameplayDebuggingControllerComponent::ToggleAIDebugView_SetView0);
-			InputComponent->BindKey(FInputChord(EKeys::One, false, false, true, false), IE_Pressed, this, &UGameplayDebuggingControllerComponent::ToggleAIDebugView_SetView1);
-			InputComponent->BindKey(FInputChord(EKeys::Two, false, false, true, false), IE_Pressed, this, &UGameplayDebuggingControllerComponent::ToggleAIDebugView_SetView2);
-			InputComponent->BindKey(FInputChord(EKeys::Three, false, false, true, false), IE_Pressed, this, &UGameplayDebuggingControllerComponent::ToggleAIDebugView_SetView3);
-			InputComponent->BindKey(FInputChord(EKeys::Four, false, false, true, false), IE_Pressed, this, &UGameplayDebuggingControllerComponent::ToggleAIDebugView_SetView4);
-			InputComponent->BindKey(FInputChord(EKeys::Five, false, false, true, false), IE_Pressed, this, &UGameplayDebuggingControllerComponent::ToggleAIDebugView_SetView5);
-			InputComponent->BindKey(FInputChord(EKeys::Six, false, false, true, false), IE_Pressed, this, &UGameplayDebuggingControllerComponent::ToggleAIDebugView_SetView6);
-			InputComponent->BindKey(FInputChord(EKeys::Seven, false, false, true, false), IE_Pressed, this, &UGameplayDebuggingControllerComponent::ToggleAIDebugView_SetView7);
-			InputComponent->BindKey(FInputChord(EKeys::Eight, false, false, true, false), IE_Pressed, this, &UGameplayDebuggingControllerComponent::ToggleAIDebugView_SetView8);
-			InputComponent->BindKey(FInputChord(EKeys::Nine, false, false, true, false), IE_Pressed, this, &UGameplayDebuggingControllerComponent::ToggleAIDebugView_SetView9);
-			InputComponent->BindKey(FInputChord(EKeys::Equals, false, false, true, false), IE_Released, this, &UGameplayDebuggingControllerComponent::CycleDetailsView);
-		}
-		else
-		{
-			InputComponent->BindKey(EKeys::NumPadZero, IE_Pressed, this, &UGameplayDebuggingControllerComponent::ToggleAIDebugView_SetView0);
-			InputComponent->BindKey(EKeys::NumPadOne, IE_Pressed, this, &UGameplayDebuggingControllerComponent::ToggleAIDebugView_SetView1);
-			InputComponent->BindKey(EKeys::NumPadTwo, IE_Pressed, this, &UGameplayDebuggingControllerComponent::ToggleAIDebugView_SetView2);
-			InputComponent->BindKey(EKeys::NumPadThree, IE_Pressed, this, &UGameplayDebuggingControllerComponent::ToggleAIDebugView_SetView3);
-			InputComponent->BindKey(EKeys::NumPadFour, IE_Pressed, this, &UGameplayDebuggingControllerComponent::ToggleAIDebugView_SetView4);
-			InputComponent->BindKey(EKeys::NumPadFive, IE_Pressed, this, &UGameplayDebuggingControllerComponent::ToggleAIDebugView_SetView5);
-			InputComponent->BindKey(EKeys::NumPadSix, IE_Pressed, this, &UGameplayDebuggingControllerComponent::ToggleAIDebugView_SetView6);
-			InputComponent->BindKey(EKeys::NumPadSeven, IE_Pressed, this, &UGameplayDebuggingControllerComponent::ToggleAIDebugView_SetView7);
-			InputComponent->BindKey(EKeys::NumPadEight, IE_Pressed, this, &UGameplayDebuggingControllerComponent::ToggleAIDebugView_SetView8);
-			InputComponent->BindKey(EKeys::NumPadNine, IE_Pressed, this, &UGameplayDebuggingControllerComponent::ToggleAIDebugView_SetView9);
-			InputComponent->BindKey(EKeys::Add, IE_Released, this, &UGameplayDebuggingControllerComponent::CycleDetailsView);
-		}
-		InputComponent->BindKey(FInputChord(EKeys::Tab, false, false, false, false), IE_Released, this, &UGameplayDebuggingControllerComponent::ToggleDebugCamera);
-		InputComponent->BindKey(FInputChord(EKeys::Tab, false, true, false, false), IE_Released, this, &UGameplayDebuggingControllerComponent::ToggleOnScreenDebugMessages);
-		InputComponent->BindKey(FInputChord(EKeys::Tilde, false, true, false, false), IE_Released, this, &UGameplayDebuggingControllerComponent::ToggleGameHUD);
+		InputComponent->BindKey(CategoryZeroBind, IE_Pressed, this, &UGameplayDebuggingControllerComponent::ToggleAIDebugView_SetView0);
+		InputComponent->BindKey(CategoryOneBind, IE_Pressed, this, &UGameplayDebuggingControllerComponent::ToggleAIDebugView_SetView1);
+		InputComponent->BindKey(CategoryTwoBind, IE_Pressed, this, &UGameplayDebuggingControllerComponent::ToggleAIDebugView_SetView2);
+		InputComponent->BindKey(CategoryThreeBind, IE_Pressed, this, &UGameplayDebuggingControllerComponent::ToggleAIDebugView_SetView3);
+		InputComponent->BindKey(CategoryFourBind, IE_Pressed, this, &UGameplayDebuggingControllerComponent::ToggleAIDebugView_SetView4);
+		InputComponent->BindKey(CategoryFiveBind, IE_Pressed, this, &UGameplayDebuggingControllerComponent::ToggleAIDebugView_SetView5);
+		InputComponent->BindKey(CategorySixBind, IE_Pressed, this, &UGameplayDebuggingControllerComponent::ToggleAIDebugView_SetView6);
+		InputComponent->BindKey(CategorySevenBind, IE_Pressed, this, &UGameplayDebuggingControllerComponent::ToggleAIDebugView_SetView7);
+		InputComponent->BindKey(CategoryEightBind, IE_Pressed, this, &UGameplayDebuggingControllerComponent::ToggleAIDebugView_SetView8);
+		InputComponent->BindKey(CategoryNineBind, IE_Pressed, this, &UGameplayDebuggingControllerComponent::ToggleAIDebugView_SetView9);
+		InputComponent->BindKey(CycleDetailsViewBind, IE_Released, this, &UGameplayDebuggingControllerComponent::CycleDetailsView);
+		InputComponent->BindKey(DebugCameraBind, IE_Released, this, &UGameplayDebuggingControllerComponent::ToggleDebugCamera);
+		InputComponent->BindKey(OnScreenDebugMessagesBind, IE_Released, this, &UGameplayDebuggingControllerComponent::ToggleOnScreenDebugMessages);
+		InputComponent->BindKey(GameHUDBind, IE_Released, this, &UGameplayDebuggingControllerComponent::ToggleGameHUD);
 	}
 #endif //!(UE_BUILD_SHIPPING || UE_BUILD_TEST)
 }
diff --git a/Engine/Source/Developer/GameplayDebugger/Private/GameplayDebuggingHUDComponent.cpp b/Engine/Source/Developer/GameplayDebugger/Private/GameplayDebuggingHUDComponent.cpp
--- a/Engine/Source/Developer/GameplayDebugger/Private/GameplayDebuggingHUDComponent.cpp
+++ b/Engine/Source/Developer/GameplayDebugger/Private/GameplayDebuggingHUDComponent.cpp
@@ -160,10 +160,6 @@ void AGameplayDebuggingHUDComponent::DrawMenu(const float X, const float Y, clas
 
 		const FString KeyDesc = ActivationKeyName != ActivationKeyDisplayName ? FString::Printf(TEXT("(%s key)"), *ActivationKeyName) : TEXT("");
 		FString HeaderDesc = FString::Printf(TEXT("Tap %s %s to close, use Numpad numbers to toggle categories"), *ActivationKeyDisplayName, *KeyDesc);
-		if (UGameplayDebuggerSettings::StaticClass()->GetDefaultObject<UGameplayDebuggerSettings>()->UseAlternateKeys())
-		{
-			HeaderDesc = FString::Printf(TEXT("Tap %s %s to close, use Alt + number to toggle categories"), *ActivationKeyDisplayName, *KeyDesc);
-		}
 
 		float HeaderWidth = 0.0f;
 		CalulateStringSize(DefaultContext, DefaultContext.Font, HeaderDesc, HeaderWidth, MaxHeight);
@@ -179,28 +175,31 @@ void AGameplayDebuggingHUDComponent::DrawMenu(const float X, const float Y, clas
 		}
 
 		{
+			static FString KeyShortcut = GDC->DebugCameraBind.GetInputText().ToString();
 			const int32 DebugCameraIndex = Categories.Add(FDebugCategoryView());
 			CategoriesWidth.AddZeroed(1);
-			Categories[DebugCameraIndex].Desc = FString::Printf(TEXT(" %s[Tab]: %s  "), GDC && GDC->GetDebugCameraController().IsValid() ? TEXT("{Green}") : TEXT("{White}"), TEXT("Debug Camera"));
+			Categories[DebugCameraIndex].Desc = FString::Printf(TEXT(" %s[%s]: %s  "), GDC && GDC->GetDebugCameraController().IsValid() ? TEXT("{Green}") : TEXT("{White}"), *KeyShortcut, TEXT("Debug Camera"));
 			float StrHeight = 0.0f;
 			CalulateStringSize(DefaultContext, DefaultContext.Font, Categories[DebugCameraIndex].Desc, CategoriesWidth[DebugCameraIndex], StrHeight);
 			TotalWidth += CategoriesWidth[DebugCameraIndex];
 			MaxHeight = FMath::Max(MaxHeight, StrHeight);
 		}
 		{
+			static FString KeyShortcut = GDC->OnScreenDebugMessagesBind.GetInputText().ToString();
 			const int32 DebugCameraIndex = Categories.Add(FDebugCategoryView());
 			CategoriesWidth.AddZeroed(1);
-			Categories[DebugCameraIndex].Desc = FString::Printf(TEXT(" %s[Ctrl+Tab]: %s  "), GEngine && GEngine->bEnableOnScreenDebugMessages ? TEXT("{Green}") : TEXT("{White}"), TEXT("DebugMessages"));
+			Categories[DebugCameraIndex].Desc = FString::Printf(TEXT(" %s[%s]: %s  "), GEngine && GEngine->bEnableOnScreenDebugMessages ? TEXT("{Green}") : TEXT("{White}"), *KeyShortcut, TEXT("DebugMessages"));
 			float StrHeight = 0.0f;
 			CalulateStringSize(DefaultContext, DefaultContext.Font, Categories[DebugCameraIndex].Desc, CategoriesWidth[DebugCameraIndex], StrHeight);
 			TotalWidth += CategoriesWidth[DebugCameraIndex];
 			MaxHeight = FMath::Max(MaxHeight, StrHeight);
 		}
 		{
+			static FString KeyShortcut = GDC->GameHUDBind.GetInputText().ToString();
 			const AHUD* GameHUD = MyPC ? MyPC->GetHUD() : NULL;
 			const int32 DebugCameraIndex = Categories.Add(FDebugCategoryView());
 			CategoriesWidth.AddZeroed(1);
-			Categories[DebugCameraIndex].Desc = FString::Printf(TEXT(" %s[Ctrl+Tilde]: %s  "), GameHUD && GameHUD->bShowHUD ? TEXT("{Green}") : TEXT("{White}"), TEXT("GameHUD"));
+			Categories[DebugCameraIndex].Desc = FString::Printf(TEXT(" %s[%s]: %s  "), GameHUD && GameHUD->bShowHUD ? TEXT("{Green}") : TEXT("{White}"), *KeyShortcut, TEXT("GameHUD"));
 			float StrHeight = 0.0f;
 			CalulateStringSize(DefaultContext, DefaultContext.Font, Categories[DebugCameraIndex].Desc, CategoriesWidth[DebugCameraIndex], StrHeight);
 			TotalWidth += CategoriesWidth[DebugCameraIndex];
@@ -713,9 +712,9 @@ void AGameplayDebuggingHUDComponent::DrawPerception(APlayerController* PC, class
 	PrintString(DefaultContext, FString::Printf(TEXT("Distance Pawn-PlayerPawn: %.1f\n"), DebugComponent->DistanceFromPlayer));
 
 	float VerticalLabelOffset = 0.f;
-	for (const FDrawDebugShapeElement& Shape : DebugComponent->PerceptionShapeElements)
+	for (const FGameplayDebuggerShapeElement& Shape : DebugComponent->PerceptionShapeElements)
 	{
-		if (Shape.GetType() == EDrawDebugShapeElement::String)
+		if (Shape.GetType() == EGameplayDebuggerShapeElement::String)
 		{
 			const FVector& Loc = Shape.Points[0];
 			const FVector ScreenLoc = DefaultContext.Canvas->Project(Loc);
@@ -723,15 +722,15 @@ void AGameplayDebuggingHUDComponent::DrawPerception(APlayerController* PC, class
 			PrintString(DefaultContext, Shape.GetFColor(), Shape.Description, ScreenLoc.X, ScreenLoc.Y + VerticalLabelOffset);
 			VerticalLabelOffset += 17;
 		}
-		else if (Shape.GetType() == EDrawDebugShapeElement::Segment)
+		else if (Shape.GetType() == EGameplayDebuggerShapeElement::Segment)
 		{
 			DrawDebugLine(World, Shape.Points[0], Shape.Points[1], Shape.GetFColor());
 		}
-		else if (Shape.GetType() == EDrawDebugShapeElement::SinglePoint)
+		else if (Shape.GetType() == EGameplayDebuggerShapeElement::SinglePoint)
 		{
 			DrawDebugSphere(World, Shape.Points[0], Shape.ThicknesOrRadius, 16, Shape.GetFColor());
 		}
-		else if (Shape.GetType() == EDrawDebugShapeElement::Cylinder)
+		else if (Shape.GetType() == EGameplayDebuggerShapeElement::Cylinder)
 		{
 			DrawDebugCylinder(World, Shape.Points[0], Shape.Points[1], Shape.ThicknesOrRadius, 16, Shape.GetFColor());
 		}
diff --git a/Engine/Source/Developer/GameplayDebugger/Private/GameplayDebuggingReplicator.cpp b/Engine/Source/Developer/GameplayDebugger/Private/GameplayDebuggingReplicator.cpp
--- a/Engine/Source/Developer/GameplayDebugger/Private/GameplayDebuggingReplicator.cpp
+++ b/Engine/Source/Developer/GameplayDebugger/Private/GameplayDebuggingReplicator.cpp
@@ -464,6 +464,8 @@ void AGameplayDebuggingReplicator::EnableTool()
 		{
 			// simulate key press
 			GDC->OnActivationKeyPressed();
+			ClientEnableTargetSelection(true, LocalPlayerOwner);
+			GetDebugComponent()->SelectTargetToDebug();
 			GDC->OnActivationKeyReleased();
 		}
 	}
diff --git a/Engine/Source/Developer/HierarchicalLODUtilities/Private/HierarchicalLODUtilities.cpp b/Engine/Source/Developer/HierarchicalLODUtilities/Private/HierarchicalLODUtilities.cpp
--- a/Engine/Source/Developer/HierarchicalLODUtilities/Private/HierarchicalLODUtilities.cpp
+++ b/Engine/Source/Developer/HierarchicalLODUtilities/Private/HierarchicalLODUtilities.cpp
@@ -106,7 +106,7 @@ bool FHierarchicalLODUtilities::BuildStaticMeshForLODActor(ALODActor* LODActor,
 {
 	if (AssetsOuter && LODActor)
 	{
-		if (!LODActor->IsDirty() || LODActor->SubActors.Num() < 2)
+		if (!LODActor->IsDirty())
 		{
 			return false;
 		}
diff --git a/Engine/Source/Developer/LogVisualizer/Private/LogVisualizerSettings.cpp b/Engine/Source/Developer/LogVisualizer/Private/LogVisualizerSettings.cpp
--- a/Engine/Source/Developer/LogVisualizer/Private/LogVisualizerSettings.cpp
+++ b/Engine/Source/Developer/LogVisualizer/Private/LogVisualizerSettings.cpp
@@ -37,6 +37,14 @@ void ULogVisualizerSettings::SavePresistentData()
 	if (bPresistentFilters)
 	{
 		PresistentFilters = FVisualLoggerFilters::Get();
+		for (int32 Index = PresistentFilters.Categories.Num() - 1; Index >= 0; --Index)
+		{
+			FCategoryFilter& Category = PresistentFilters.Categories[Index];
+			if (Category.bIsInUse == false)
+			{
+				PresistentFilters.Categories.RemoveAt(Index);
+			}
+		}
 	}
 	else
 	{
@@ -57,11 +65,16 @@ void ULogVisualizerSettings::LoadPresistentData()
 {
 	if (bPresistentFilters)
 	{
-		static_cast<FVisualLoggerFiltersData&>(FVisualLoggerFilters::Get()) = PresistentFilters;
+		for (int32 Index = PresistentFilters.Categories.Num() - 1; Index >= 0; --Index)
+		{
+			FCategoryFilter& Category = PresistentFilters.Categories[Index];
+			Category.bIsInUse = false;
+		}
+		FVisualLoggerFilters::Get().InitWith(PresistentFilters);
 	}
 	else
 	{
-		FVisualLoggerFilters::Get() = FVisualLoggerFilters();
+		FVisualLoggerFilters::Get().Reset();
 	}
 }
 
@@ -129,9 +142,10 @@ void FVisualLoggerFilters::AddCategory(FString InName, ELogVerbosity::Type InVer
 	const int32 Num = Categories.Num();
 	for (int32 Index = 0; Index < Num; ++Index)
 	{
-		const FCategoryFilter& Filter = Categories[Index];
+		FCategoryFilter& Filter = Categories[Index];
 		if (Filter.CategoryName == InName)
 		{
+			Filter.bIsInUse = true;
 			return;
 		}
 	}
@@ -140,6 +154,7 @@ void FVisualLoggerFilters::AddCategory(FString InName, ELogVerbosity::Type InVer
 	Filter.CategoryName = InName;
 	Filter.LogVerbosity = InVerbosity;
 	Filter.Enabled = true;
+	Filter.bIsInUse = true;
 	Categories.Add(Filter);
 
 	FastCategoryFilterMap.Reset(); // we need to recreate cache - pointers can be broken
@@ -250,14 +265,40 @@ void FVisualLoggerFilters::EnableAllCategories()
 
 void FVisualLoggerFilters::Reset()
 {
-	FastCategoryFilterMap.Reset();
+	if (ULogVisualizerSettings::StaticClass()->GetDefaultObject<ULogVisualizerSettings>()->bResetDataWithNewSession == false)
+	{
+		for (int32 Index = Categories.Num() - 1; Index >= 0; --Index)
+		{
+			Categories[Index].bIsInUse = false;
+		}
+	}
+	else
+	{
+		FastCategoryFilterMap.Reset();
+		Categories.Reset();
+	}
+
 	SearchBoxFilter = FString();
 	ObjectNameFilter = FString();
 
-	Categories.Reset();
 	SelectedClasses.Reset();
 }
 
+void FVisualLoggerFilters::InitWith(const FVisualLoggerFiltersData& NewFiltersData)
+{
+	SearchBoxFilter = NewFiltersData.SearchBoxFilter;
+	ObjectNameFilter = NewFiltersData.ObjectNameFilter;
+	Categories = NewFiltersData.Categories;
+	SelectedClasses = NewFiltersData.SelectedClasses;
+
+	FastCategoryFilterMap.Reset(); // we need to recreate cache - pointers can be broken
+	FastCategoryFilterMap.Reserve(Categories.Num());
+	for (int32 Index = 0; Index < Categories.Num(); Index++)
+	{
+		FastCategoryFilterMap.Add(*Categories[Index].CategoryName, &Categories[Index]);
+	}
+}
+
 void FVisualLoggerFilters::DisableGraphData(FName GraphName, FName DataName, bool SetAsDisabled)
 {
 	const FName FullName = *(GraphName.ToString() + TEXT("$") + DataName.ToString());
diff --git a/Engine/Source/Developer/LogVisualizer/Private/STimeline.cpp b/Engine/Source/Developer/LogVisualizer/Private/STimeline.cpp
--- a/Engine/Source/Developer/LogVisualizer/Private/STimeline.cpp
+++ b/Engine/Source/Developer/LogVisualizer/Private/STimeline.cpp
@@ -60,21 +60,29 @@ void STimeline::UpdateVisibility()
 {
 	ULogVisualizerSettings* Settings = ULogVisualizerSettings::StaticClass()->GetDefaultObject<ULogVisualizerSettings>();
 
-	FVisualLoggerDBRow& DataRow = FVisualLoggerDatabase::Get().GetRowByName(GetName());
-	const TArray<FVisualLogDevice::FVisualLogEntryItem>& Entries = DataRow.GetItems();
+	if (FVisualLoggerDatabase::Get().ContainsRowByName(GetName()))
+	{
+		FVisualLoggerDBRow& DataRow = FVisualLoggerDatabase::Get().GetRowByName(GetName());
+		const TArray<FVisualLogDevice::FVisualLogEntryItem>& Entries = DataRow.GetItems();
 
-	const bool bJustIgnore = Settings->bIgnoreTrivialLogs && Entries.Num() <= Settings->TrivialLogsThreshold;
-	const bool bVisibleByOwnerClasses = FVisualLoggerFilters::Get().MatchObjectName(OwnerClassName.ToString());
-	const bool bIsCollapsed = bJustIgnore || DataRow.GetNumberOfHiddenItems() == Entries.Num() || (SearchFilter.Len() > 0 && OwnerName.ToString().Find(SearchFilter) == INDEX_NONE);
+		const bool bJustIgnore = Settings->bIgnoreTrivialLogs && Entries.Num() <= Settings->TrivialLogsThreshold;
+		const bool bVisibleByOwnerClasses = FVisualLoggerFilters::Get().MatchObjectName(OwnerClassName.ToString());
+		const bool bIsCollapsed = bJustIgnore || DataRow.GetNumberOfHiddenItems() == Entries.Num() || (SearchFilter.Len() > 0 && OwnerName.ToString().Find(SearchFilter) == INDEX_NONE);
 
-	const bool bIsHidden = bIsCollapsed || !bVisibleByOwnerClasses;
-	SetVisibility(bIsHidden ? EVisibility::Collapsed : EVisibility::Visible);
-	if (bIsCollapsed)
+		const bool bIsHidden = bIsCollapsed || !bVisibleByOwnerClasses;
+		SetVisibility(bIsHidden ? EVisibility::Collapsed : EVisibility::Visible);
+		if (bIsCollapsed)
+		{
+			Owner->SetSelectionState(SharedThis(this), false, false);
+		}
+
+		FVisualLoggerDatabase::Get().SetRowVisibility(GetName(), !bIsHidden);
+	}
+	else
 	{
 		Owner->SetSelectionState(SharedThis(this), false, false);
+		FVisualLoggerDatabase::Get().SetRowVisibility(GetName(), false);
 	}
-
-	FVisualLoggerDatabase::Get().SetRowVisibility(GetName(), !bIsHidden);
 }
 
 void STimeline::OnFiltersSearchChanged(const FText& Filter)
@@ -192,10 +200,7 @@ void STimeline::Construct(const FArguments& InArgs, TSharedPtr<FVisualLoggerTime
 
 STimeline::~STimeline()
 {
-	if (ULogVisualizerSettings::StaticClass() && ULogVisualizerSettings::StaticClass()->GetDefaultObject<ULogVisualizerSettings>())
-	{
-		ULogVisualizerSettings::StaticClass()->GetDefaultObject<ULogVisualizerSettings>()->OnSettingChanged().RemoveAll(this);
-	}
+	ULogVisualizerSettings::StaticClass()->GetDefaultObject<ULogVisualizerSettings>()->OnSettingChanged().RemoveAll(this);
 	FVisualLoggerDatabase::Get().GetEvents().OnNewItem.RemoveAll(this);
 	FVisualLoggerDatabase::Get().GetEvents().OnRowSelectionChanged.RemoveAll(this);
 }
diff --git a/Engine/Source/Developer/LogVisualizer/Private/SVisualLogger.cpp b/Engine/Source/Developer/LogVisualizer/Private/SVisualLogger.cpp
--- a/Engine/Source/Developer/LogVisualizer/Private/SVisualLogger.cpp
+++ b/Engine/Source/Developer/LogVisualizer/Private/SVisualLogger.cpp
@@ -46,7 +46,7 @@ DECLARE_DELEGATE_TwoParams(FOnWorldChanged, UWorld*, UWorld*);
 *****************************************************************************/
 namespace
 {
-	static UWorld* GetWorldForGivenObject(const class UObject* Object)
+	static UWorld* GetWorldForGivenObject(const UObject* Object)
 	{
 		UWorld* World = Object ? GEngine->GetWorldFromContextObject(Object, false) : nullptr;
 #if WITH_EDITOR
@@ -79,10 +79,14 @@ SVisualLogger::SVisualLogger()
 		FVisualLoggerDevice(SVisualLogger* InVisualLogger, FOnWorldChanged OnWorldChangedDelegate) : VisualLoggerWidget(InVisualLogger), LastUsedWorld(nullptr), OnWorldChanged(OnWorldChangedDelegate) {}
 
 		virtual ~FVisualLoggerDevice(){}
-		virtual void Serialize(const class UObject* LogOwner, FName OwnerName, FName OwnerClassName, const FVisualLogEntry& LogEntry) override
+		virtual void Serialize(const UObject* LogOwner, FName OwnerName, FName OwnerClassName, const FVisualLogEntry& LogEntry) override
 		{
 			VisualLoggerWidget->OnNewLogEntry(FVisualLogDevice::FVisualLogEntryItem(OwnerName, OwnerClassName, LogEntry));
-			UWorld* CurrentWorld = GetWorldForGivenObject(LogOwner);
+			UWorld* CurrentWorld = nullptr;
+			if (FVisualLogger::Get().GetObjectToWorldMap().Contains(LogOwner))
+			{
+				CurrentWorld = const_cast<UWorld*>(FVisualLogger::Get().GetObjectToWorldMap()[LogOwner].Get());
+			}
 			if (LastUsedWorld != CurrentWorld)
 			{
 				OnWorldChanged.ExecuteIfBound(LastUsedWorld, CurrentWorld);
@@ -678,6 +682,7 @@ void SVisualLogger::ResetData()
 
 	FLogVisualizer::Get().Reset();
 	FVisualLoggerDatabase::Get().Reset();
+
 	FVisualLoggerFilters::Get().Reset();
 
 	if (MainView.IsValid())
@@ -685,7 +690,7 @@ void SVisualLogger::ResetData()
 		MainView->ResetData();
 	}
 
-	if (VisualLoggerFilters.IsValid())
+	if (VisualLoggerFilters.IsValid() && ULogVisualizerSettings::StaticClass()->GetDefaultObject<ULogVisualizerSettings>()->bResetDataWithNewSession)
 	{
 		VisualLoggerFilters->ResetData();
 	}
diff --git a/Engine/Source/Developer/LogVisualizer/Private/SVisualLoggerFilters.cpp b/Engine/Source/Developer/LogVisualizer/Private/SVisualLoggerFilters.cpp
--- a/Engine/Source/Developer/LogVisualizer/Private/SVisualLoggerFilters.cpp
+++ b/Engine/Source/Developer/LogVisualizer/Private/SVisualLoggerFilters.cpp
@@ -37,7 +37,6 @@ void SVisualLoggerFilters::Construct(const FArguments& InArgs, const TSharedRef<
 			.Text(LOCTEXT("GraphFilters", "Graph Filters"))
 		];
 
-	GraphsFilterCombo->SetVisibility(CachedDatasPerGraph.Num() ? EVisibility::Visible : EVisibility::Collapsed);
 
 	FilterBox->AddSlot()
 		.Padding(3, 3)
@@ -60,6 +59,13 @@ void SVisualLoggerFilters::Construct(const FArguments& InArgs, const TSharedRef<
 			]
 		];
 
+	for (auto& CurrentCategory : FVisualLoggerFilters::Get().Categories)
+	{
+		AddFilterCategory(CurrentCategory.CategoryName, (ELogVerbosity::Type)CurrentCategory.LogVerbosity, false);
+	}
+
+	GraphsFilterCombo->SetVisibility(CachedDatasPerGraph.Num() ? EVisibility::Visible : EVisibility::Collapsed);
+
 	FVisualLoggerFilters::Get().OnFilterCategoryAdded.AddRaw(this, &SVisualLoggerFilters::OnFilterCategoryAdded);
 	FVisualLoggerFilters::Get().OnFilterCategoryRemoved.AddRaw(this, &SVisualLoggerFilters::OnFilterCategoryRemoved);
 	FVisualLoggerDatabase::Get().GetEvents().OnItemSelectionChanged.AddRaw(this, &SVisualLoggerFilters::OnItemsSelectionChanged);
@@ -198,17 +204,10 @@ void SVisualLoggerFilters::GraphFilterCategoryClicked(FName GraphName)
 {
 	const bool bNewSet = !IsGraphFilterCategoryInUse(GraphName);
 
-	for (auto Iter(FVisualLoggerGraphsDatabase::Get().GetConstOwnersIterator()); Iter; ++Iter)
+	for (auto Iter(CachedDatasPerGraph[GraphName].CreateConstIterator()); Iter; ++Iter)
 	{
-		const FName OwnerName = Iter.Key();
-		if (FVisualLoggerGraphsDatabase::Get().ContainsGraphByName(OwnerName, GraphName))
-		{
-			FVisualLoggerGraph& Graph = FVisualLoggerGraphsDatabase::Get().GetGraphByName(OwnerName, GraphName);
-			for (auto DataIt(Graph.GetConstDataIterator()); DataIt; ++DataIt)
-			{
-				FVisualLoggerFilters::Get().DisableGraphData(GraphName, DataIt->DataName, !bNewSet);
-			}
-		}
+		const FName& DataName = *Iter;
+		FVisualLoggerFilters::Get().DisableGraphData(GraphName, DataName, !bNewSet);
 	}
 
 	FLogVisualizer::Get().GetEvents().OnFiltersChanged.Broadcast();
@@ -218,17 +217,11 @@ void SVisualLoggerFilters::GraphFilterCategoryClicked(FName GraphName)
 bool SVisualLoggerFilters::IsGraphFilterCategoryInUse(FName GraphName) const
 {
 	bool bInUse = false;
-	for (auto Iter(FVisualLoggerGraphsDatabase::Get().GetConstOwnersIterator()); Iter; ++Iter)
+
+	for (auto Iter(CachedDatasPerGraph[GraphName].CreateConstIterator()); Iter; ++Iter)
 	{
-		const FName OwnerName = Iter.Key();
-		if (FVisualLoggerGraphsDatabase::Get().ContainsGraphByName(OwnerName, GraphName))
-		{
-			FVisualLoggerGraph& Graph = FVisualLoggerGraphsDatabase::Get().GetGraphByName(OwnerName, GraphName);
-			for (auto DataIt(Graph.GetConstDataIterator()); DataIt; ++DataIt)
-			{
-				bInUse |= FVisualLoggerFilters::Get().IsGraphDataDisabled(GraphName, DataIt->DataName) == false;
-			}
-		}
+		const FName& DataName = *Iter;
+		bInUse |= FVisualLoggerFilters::Get().IsGraphDataDisabled(GraphName, DataName) == false;
 	}
 
 	return bInUse;
@@ -308,7 +301,7 @@ uint32 SVisualLoggerFilters::GetCategoryIndex(const FString& InFilterName) const
 	return INDEX_NONE;
 }
 
-void SVisualLoggerFilters::OnFilterCategoryAdded(FString InName, ELogVerbosity::Type InVerbosity)
+void SVisualLoggerFilters::AddFilterCategory(FString InName, ELogVerbosity::Type InVerbosity, bool bMarkAsInUse)
 {
 	int32 CharIndex = INDEX_NONE;
 	if (InName.FindChar('$', CharIndex) == true)
@@ -318,6 +311,7 @@ void SVisualLoggerFilters::OnFilterCategoryAdded(FString InName, ELogVerbosity::
 		if (ensure(GroupAndName.Num() == 2))
 		{
 			CachedGraphFilters.FindOrAdd(*GroupAndName[0]).AddUnique(GroupAndName[1]);
+			CachedDatasPerGraph.FindOrAdd(*GroupAndName[0]).AddUnique(*GroupAndName[1]);
 		}
 	}
 	else
@@ -345,6 +339,23 @@ void SVisualLoggerFilters::OnFilterCategoryAdded(FString InName, ELogVerbosity::
 			];
 	}
 
+	if (bMarkAsInUse)
+	{
+		FVisualLoggerFilters& PresistentFilters = FVisualLoggerFilters::Get();
+		for (int32 Index = PresistentFilters.Categories.Num() - 1; Index >= 0; --Index)
+		{
+			FCategoryFilter& Category = PresistentFilters.Categories[Index];
+			if (Category.CategoryName == InName)
+			{
+				Category.bIsInUse = true;
+			}
+		}
+	}
+}
+
+void SVisualLoggerFilters::OnFilterCategoryAdded(FString InName, ELogVerbosity::Type InVerbosity)
+{
+	AddFilterCategory(InName, InVerbosity, false);
 	GraphsFilterCombo->SetVisibility(CachedGraphFilters.Num() > 0 ? EVisibility::Visible : EVisibility::Collapsed);
 }
 
diff --git a/Engine/Source/Developer/LogVisualizer/Private/SVisualLoggerFilters.h b/Engine/Source/Developer/LogVisualizer/Private/SVisualLoggerFilters.h
--- a/Engine/Source/Developer/LogVisualizer/Private/SVisualLoggerFilters.h
+++ b/Engine/Source/Developer/LogVisualizer/Private/SVisualLoggerFilters.h
@@ -35,6 +35,7 @@ public:
 protected:
 	//void OnNewItemHandler(const FVisualLoggerDBRow& BDRow, int32 ItemIndex);
 
+	void AddFilterCategory(FString, ELogVerbosity::Type, bool bMarkAsInUse);
 	void OnFilterCategoryAdded(FString, ELogVerbosity::Type);
 	void OnFilterCategoryRemoved(FString);
 	void OnItemsSelectionChanged(const FVisualLoggerDBRow& ChangedRow, int32 SelectedItemIndex);
