Fix Mac non-unity compile

diff --git a/Engine/Plugins/Developer/GameplayDebugger/Source/GameplayDebuggerModule/Private/GameplayDebuggerReplicator.cpp b/Engine/Plugins/Developer/GameplayDebugger/Source/GameplayDebuggerModule/Private/GameplayDebuggerReplicator.cpp
--- a/Engine/Plugins/Developer/GameplayDebugger/Source/GameplayDebuggerModule/Private/GameplayDebuggerReplicator.cpp
+++ b/Engine/Plugins/Developer/GameplayDebugger/Source/GameplayDebuggerModule/Private/GameplayDebuggerReplicator.cpp
@@ -12,6 +12,7 @@
 #include "GameplayDebuggerReplicator.h"
 #include "BehaviorTreeDelegates.h"
 #if WITH_EDITOR
+#include "Editor.h"
 #include "Editor/EditorEngine.h"
 #include "Editor/BehaviorTreeEditor/Public/BehaviorTreeEditorModule.h"
 #include "LevelEditorViewport.h"
